EvalGraph : Base level classes for node graph evaluations.

diff --git a/Engine/Plugins/Experimental/EvalGraph/EvalGraph.uplugin b/Engine/Plugins/Experimental/EvalGraph/EvalGraph.uplugin
new file mode 100644
--- /dev/null
+++ b/Engine/Plugins/Experimental/EvalGraph/EvalGraph.uplugin
@@ -0,0 +1,27 @@
+{
+    "FileVersion" : 3,
+    "Version" : 1,
+    "VersionName" : "0.1",
+    "FriendlyName" : "Editor Evaluation Graph",
+    "Description" : "Editor Evaluation Graph",
+    "Category": "Editor",
+    "CreatedBy" : "Epic Games, Inc.",
+    "CreatedByURL" : "http://epicgames.com",
+    "DocsURL" : "",
+    "MarketplaceURL" : "",
+    "SupportURL" : "",
+    "EnabledByDefault" : false,
+    "CanContainContent" : true,
+    "IsBetaVersion" : true,
+	"IsExperimentalVersion": true,
+    "Installed" : false,
+	"Modules": [
+		{
+			"Name": "EvalGraph",
+			"Type": "Runtime",
+			"LoadingPhase": "PreDefault"
+		},
+	],
+	"Plugins": [
+	]
+}
diff --git a/Engine/Plugins/Experimental/EvalGraph/Source/EvalGraph/EvalGraph.Build.cs b/Engine/Plugins/Experimental/EvalGraph/Source/EvalGraph/EvalGraph.Build.cs
new file mode 100644
--- /dev/null
+++ b/Engine/Plugins/Experimental/EvalGraph/Source/EvalGraph/EvalGraph.Build.cs
@@ -0,0 +1,22 @@
+// Copyright Epic Games, Inc. All Rights Reserved.
+
+namespace UnrealBuildTool.Rules
+{
+	public class EvalGraph : ModuleRules
+	{
+        public EvalGraph(ReadOnlyTargetRules Target) : base(Target)
+		{
+			PrivateIncludePaths.Add("EvalGraph/Private");
+            PublicIncludePaths.Add(ModuleDirectory + "/Public");
+
+			PublicDependencyModuleNames.AddRange(
+				new string[]
+				{
+					"Core",
+					"CoreUObject",
+					"Chaos",
+				}
+				);
+		}
+	}
+}
diff --git a/Engine/Plugins/Experimental/EvalGraph/Source/EvalGraph/Private/EvalGraph/EvalGraph.cpp b/Engine/Plugins/Experimental/EvalGraph/Source/EvalGraph/Private/EvalGraph/EvalGraph.cpp
new file mode 100644
--- /dev/null
+++ b/Engine/Plugins/Experimental/EvalGraph/Source/EvalGraph/Private/EvalGraph/EvalGraph.cpp
@@ -0,0 +1,69 @@
+// Copyright Epic Games, Inc. All Rights Reserved.
+
+#pragma once
+#include "EvalGraph/EvalGraph.h"
+
+#include "EvalGraph/EvalGraphInputOutput.h"
+#include "EvalGraph/EvalGraphConnectionTypes.h"
+#include "EvalGraph/EvalGraphNodeParameters.h"
+#include "EvalGraph/EvalGraphNode.h"
+
+namespace Eg
+{
+
+	FGraph::FGraph(FGuid InGuid)
+		: Guid(InGuid)
+	{}
+
+	void FGraph::RemoveNode(TSharedPtr<FNode> Node)
+	{
+		for (FConnectionTypeBase* Output : Node->GetOutputs())
+		{
+			if (Output)
+			{
+				for (FConnectionTypeBase* Input : Output->GetBaseInputs())
+				{
+					if (Input)
+					{
+						Disconnect(Output, Input);
+					}
+				}
+			}
+		}
+		for (FConnectionTypeBase* Input : Node->GetInputs())
+		{
+			if (Input)
+			{
+				TArray<FConnectionTypeBase*> Outputs = Input->GetBaseOutputs();
+				for (FConnectionTypeBase* Out : Outputs)
+				{
+					if (Out)
+					{
+						Disconnect(Out, Input);
+					}
+				}
+			}
+		}
+		Nodes.Remove(Node);
+	}
+
+	void FGraph::Connect(FConnectionTypeBase* Input, FConnectionTypeBase* Output)
+	{
+		if (ensure(Input && Output))
+		{
+			Input->AddConnection(Output);
+			Output->AddConnection(Input);
+			Connections.Add(FConnection(Input->GetGuid(), Output->GetGuid()));
+		}
+	}
+
+	void FGraph::Disconnect(FConnectionTypeBase* Input, FConnectionTypeBase* Output)
+	{
+		Input->RemoveConnection(Output);
+		Output->RemoveConnection(Input);
+		Connections.RemoveSwap(FConnection(Input->GetGuid(), Output->GetGuid()));
+	}
+
+
+}
+
diff --git a/Engine/Plugins/Experimental/EvalGraph/Source/EvalGraph/Private/EvalGraph/EvalGraphInputOutput.cpp b/Engine/Plugins/Experimental/EvalGraph/Source/EvalGraph/Private/EvalGraph/EvalGraphInputOutput.cpp
new file mode 100644
--- /dev/null
+++ b/Engine/Plugins/Experimental/EvalGraph/Source/EvalGraph/Private/EvalGraph/EvalGraphInputOutput.cpp
@@ -0,0 +1,72 @@
+// Copyright Epic Games, Inc. All Rights Reserved.
+
+#pragma once
+#include "EvalGraph/EvalGraphInputOutput.h"
+
+#include "EvalGraph/EvalGraphConnectionTypes.h"
+#include "EvalGraph/EvalGraphNodeParameters.h"
+#include "EvalGraph/EvalGraphNode.h"
+
+namespace Eg
+{
+
+	FConnectionTypeBase::FConnectionTypeBase(EGraphConnectionType InType, FNode* InOwningNode, FGuid InGuid)
+		: Type(InType)
+		, Guid(InGuid)
+		, OwningNode(InOwningNode)
+	{}
+
+	void FConnectionTypeBase::AddBaseInput(FNode* InNode, FConnectionTypeBase* That) 
+	{ InNode->AddBaseInput(That); }
+	void FConnectionTypeBase::AddBaseOutput(FNode* InNode, FConnectionTypeBase* That) 
+	{ InNode->AddBaseOutput(That); }
+
+
+	template<class T>
+	FInput<T>::FInput(const FInputParameters<T>& Param, FGuid InGuid )
+		: FConnectionTypeBase(GraphConnectionType<T>(), Param.Owner, InGuid)
+		, Name(Param.Name)
+		, Default(Param.Default)
+		, Connection(nullptr)
+	{
+		Super::AddBaseInput(Param.Owner, this);
+	}
+
+
+	template<class T>
+	FOutput<T>::FOutput(const FOutputParameters& Param, FGuid InGuid )
+		: FConnectionTypeBase(GraphConnectionType<T>(), Param.Owner, InGuid)
+		, Name(Param.Name)
+	{
+		Super::AddBaseOutput(Param.Owner, this);
+	}
+
+
+	template<class T>
+	void FOutput<T>::SetValue(T InVal, const FContext& Context)
+	{
+		CacheKey = Context.GetTypeHash();
+		Cache.Data = InVal;
+	}
+
+	template<class T>
+	T FOutput<T>::Evaluate(const FContext& Context)
+	{
+		if (CacheKey!=Context.GetTypeHash())
+		{
+			OwningNode->Evaluate(Context, this);
+		}
+		ensure(CacheKey==Context.GetTypeHash());
+		return Cache.Data;
+	}
+
+
+#define EVAL_GRAPH_CONNECTION_TYPE(a,A) template class FOutput<a>;
+#include "../../Public/EvalGraph/EvalGraphConnectionTypeValues.inl"
+#undef EVAL_GRAPH_CONNECTION_TYPE
+
+#define EVAL_GRAPH_CONNECTION_TYPE(a,A) template class FInput<a>;
+#include "../../Public/EvalGraph/EvalGraphConnectionTypeValues.inl"
+#undef EVAL_GRAPH_CONNECTION_TYPE
+}
+
diff --git a/Engine/Plugins/Experimental/EvalGraph/Source/EvalGraph/Private/EvalGraph/EvalGraphNode.cpp b/Engine/Plugins/Experimental/EvalGraph/Source/EvalGraph/Private/EvalGraph/EvalGraphNode.cpp
new file mode 100644
--- /dev/null
+++ b/Engine/Plugins/Experimental/EvalGraph/Source/EvalGraph/Private/EvalGraph/EvalGraphNode.cpp
@@ -0,0 +1,11 @@
+// Copyright Epic Games, Inc. All Rights Reserved.
+
+#include "EvalGraph/EvalGraphNode.h"
+
+#include "EvalGraph/EvalGraphNode.h"
+
+namespace Eg
+{
+
+}
+
diff --git a/Engine/Plugins/Experimental/EvalGraph/Source/EvalGraph/Private/EvalGraph/EvalGraphNodeParameters.cpp b/Engine/Plugins/Experimental/EvalGraph/Source/EvalGraph/Private/EvalGraph/EvalGraphNodeParameters.cpp
new file mode 100644
--- /dev/null
+++ b/Engine/Plugins/Experimental/EvalGraph/Source/EvalGraph/Private/EvalGraph/EvalGraphNodeParameters.cpp
@@ -0,0 +1,4 @@
+// Copyright Epic Games, Inc. All Rights Reserved.
+
+#include "EvalGraph/EvalGraphNodeParameters.h"
+
diff --git a/Engine/Plugins/Experimental/EvalGraph/Source/EvalGraph/Private/EvalGraph/EvalGraphPlugin.cpp b/Engine/Plugins/Experimental/EvalGraph/Source/EvalGraph/Private/EvalGraph/EvalGraphPlugin.cpp
new file mode 100644
--- /dev/null
+++ b/Engine/Plugins/Experimental/EvalGraph/Source/EvalGraph/Private/EvalGraph/EvalGraphPlugin.cpp
@@ -0,0 +1,32 @@
+// Copyright Epic Games, Inc. All Rights Reserved.
+
+#include "EvalGraph/EvalGraphPlugin.h"
+
+#include "CoreMinimal.h"
+#include "Modules/ModuleManager.h"
+
+
+
+class FEvalGraphPlugin : public IEvalGraphPlugin
+{
+	/** IModuleInterface implementation */
+	virtual void StartupModule() override;
+	virtual void ShutdownModule() override;
+};
+
+IMPLEMENT_MODULE( FEvalGraphPlugin, EvalGraph )
+
+
+void FEvalGraphPlugin::StartupModule()
+{
+	
+}
+
+
+void FEvalGraphPlugin::ShutdownModule()
+{
+	
+}
+
+
+
diff --git a/Engine/Plugins/Experimental/EvalGraph/Source/EvalGraph/Private/EvalGraph/EvalGraphPluginPrivate.h b/Engine/Plugins/Experimental/EvalGraph/Source/EvalGraph/Private/EvalGraph/EvalGraphPluginPrivate.h
new file mode 100644
--- /dev/null
+++ b/Engine/Plugins/Experimental/EvalGraph/Source/EvalGraph/Private/EvalGraph/EvalGraphPluginPrivate.h
@@ -0,0 +1,12 @@
+// Copyright Epic Games, Inc. All Rights Reserved.
+
+#pragma once
+
+#include "CoreMinimal.h"
+#include "Stats/Stats.h"
+
+// You should place include statements to your module's private header files here.  You only need to
+// add includes for headers that are used in most of your module's source files though.
+
+
+DECLARE_STATS_GROUP(TEXT("EvalGraph"), STATGROUP_EvalGraph, STATCAT_Advanced);
diff --git a/Engine/Plugins/Experimental/EvalGraph/Source/EvalGraph/Public/EvalGraph/EvalGraph.h b/Engine/Plugins/Experimental/EvalGraph/Source/EvalGraph/Public/EvalGraph/EvalGraph.h
new file mode 100644
--- /dev/null
+++ b/Engine/Plugins/Experimental/EvalGraph/Source/EvalGraph/Public/EvalGraph/EvalGraph.h
@@ -0,0 +1,46 @@
+// Copyright Epic Games, Inc. All Rights Reserved.
+
+#pragma once
+
+#include "CoreMinimal.h"
+
+
+namespace Eg
+{
+	class FNode;
+	class FConnectionTypeBase;
+
+	typedef TTuple<FGuid, FGuid> FConnection;
+
+
+	//
+	//
+	//
+	class EVALGRAPH_API FGraph
+	{
+
+		FGuid  Guid;
+		TArray< TSharedPtr<FNode> > Nodes;
+		TArray< FConnection > Connections;
+		TMap<FGuid, TSharedPtr<FConnectionTypeBase> > ConnectionMap;
+
+	public:
+		FGraph(FGuid InGuid = FGuid::NewGuid());
+
+		int NumNodes() { return Nodes.Num(); }
+
+		template<class T> TSharedPtr<T> AddNode(T* InNode)
+		{
+			TSharedPtr<T> NewNode(InNode);
+			Nodes.AddUnique(NewNode);
+			return NewNode;
+		}
+
+		void RemoveNode(TSharedPtr<FNode> Node);
+
+		void Connect(FConnectionTypeBase* Input, FConnectionTypeBase* Output);
+		void Disconnect(FConnectionTypeBase* Input, FConnectionTypeBase* Output);
+	};
+
+}
+
diff --git a/Engine/Plugins/Experimental/EvalGraph/Source/EvalGraph/Public/EvalGraph/EvalGraphConnectionTypeValues.inl b/Engine/Plugins/Experimental/EvalGraph/Source/EvalGraph/Public/EvalGraph/EvalGraphConnectionTypeValues.inl
new file mode 100644
--- /dev/null
+++ b/Engine/Plugins/Experimental/EvalGraph/Source/EvalGraph/Public/EvalGraph/EvalGraphConnectionTypeValues.inl
@@ -0,0 +1,19 @@
+// Copyright Epic Games, Inc. All Rights Reserved.
+
+// usage
+//
+// General purpose Scene Graph Type definitions
+
+#ifndef EVAL_GRAPH_CONNECTION_TYPE
+#error EVAL_GRAPH_CONNECTION_TYPE macro is undefined.
+#endif
+
+// NOTE: new types must be added at the bottom to keep serialization from breaking
+
+EVAL_GRAPH_CONNECTION_TYPE(int, Integer)
+EVAL_GRAPH_CONNECTION_TYPE(TSharedPtr<FManagedArrayCollection>, ManagedArrayCollection)
+
+// NOTE: new types must be added at the bottom to keep serialization from breaking
+
+
+#undef EVAL_GRAPH_CONNECTION_TYPE
diff --git a/Engine/Plugins/Experimental/EvalGraph/Source/EvalGraph/Public/EvalGraph/EvalGraphConnectionTypes.h b/Engine/Plugins/Experimental/EvalGraph/Source/EvalGraph/Public/EvalGraph/EvalGraphConnectionTypes.h
new file mode 100644
--- /dev/null
+++ b/Engine/Plugins/Experimental/EvalGraph/Source/EvalGraph/Public/EvalGraph/EvalGraphConnectionTypes.h
@@ -0,0 +1,59 @@
+// Copyright Epic Games, Inc. All Rights Reserved.
+#pragma once
+#include "GeometryCollection/ManagedArrayCollection.h"
+
+namespace Eg
+{
+	// ---------------------------------------------------------
+	//
+	// General purpose EManagedArrayType definition. 
+	// This defines things like 
+	//     EGraphConnectionType::FManagedArrayCollectionType
+	// see EvalGraphConnectionTypeValues.inl for specific types.
+	//
+#define EVAL_GRAPH_CONNECTION_TYPE(a,A) F##A##Type,
+	enum class EGraphConnectionType : uint8
+	{
+		FNoneType,
+#include "EvalGraphConnectionTypeValues.inl"
+	};
+#undef EVAL_GRAPH_CONNECTION_TYPE
+
+	// ---------------------------------------------------------
+	//  GraphConnectionType<T>
+	//    Templated function to return a EGraphConnectionType.
+	//
+	template<class T> inline EGraphConnectionType GraphConnectionType();
+#define EVAL_GRAPH_CONNECTION_TYPE(a,A) template<> inline EGraphConnectionType GraphConnectionType<a>() { return EGraphConnectionType::F##A##Type; }
+#include "EvalGraphConnectionTypeValues.inl"
+#undef EVAL_GRAPH_CONNECTION_TYPE
+
+	// ---------------------------------------------------------
+	//  GraphConnectionTypeName
+	//    Templated function to return a EGraphConnectionType.
+	//
+	template<class T> inline FName GraphConnectionTypeName();
+#define EVAL_GRAPH_CONNECTION_TYPE(a,A) template<> inline FName GraphConnectionTypeName<a>() { return FName(TEXT(#A)); }
+#include "EvalGraphConnectionTypeValues.inl"
+#undef EVAL_GRAPH_CONNECTION_TYPE
+
+	// ---------------------------------------------------------
+	//  void*
+	//     Returns a new EGraphConnectionType pointer based on 
+	//     passed type.
+	//
+	inline void* NewGraphValueType(EGraphConnectionType ValueType)
+	{
+		switch (ValueType)
+		{
+#define EVAL_GRAPH_CONNECTION_TYPE(a,A)	case EGraphConnectionType::F##A##Type:\
+		return new a();
+#include "EvalGraphConnectionTypeValues.inl"
+#undef EVAL_GRAPH_CONNECTION_TYPE
+		}
+		check(false);
+		return nullptr;
+	}
+
+}
+
diff --git a/Engine/Plugins/Experimental/EvalGraph/Source/EvalGraph/Public/EvalGraph/EvalGraphInputOutput.h b/Engine/Plugins/Experimental/EvalGraph/Source/EvalGraph/Public/EvalGraph/EvalGraphInputOutput.h
new file mode 100644
--- /dev/null
+++ b/Engine/Plugins/Experimental/EvalGraph/Source/EvalGraph/Public/EvalGraph/EvalGraphInputOutput.h
@@ -0,0 +1,176 @@
+// Copyright Epic Games, Inc. All Rights Reserved.
+
+#pragma once
+
+#include "CoreMinimal.h"
+#include "EvalGraph/EvalGraphNodeParameters.h"
+
+namespace Eg
+{
+	class FNode;
+
+	//
+	// Input Output Base
+	//
+
+	class EVALGRAPH_API FConnectionTypeBase
+	{
+
+	protected:
+		EGraphConnectionType Type;
+		FGuid  Guid;
+		FNode* OwningNode = nullptr;
+
+		friend class FNode;
+		static void AddBaseInput(FNode* InNode, FConnectionTypeBase*);
+		static void AddBaseOutput(FNode* InNode, FConnectionTypeBase*);
+
+	public:
+		FConnectionTypeBase(EGraphConnectionType InType, FNode* OwningNode = nullptr, FGuid InGuid = FGuid::NewGuid());
+		virtual ~FConnectionTypeBase() {};
+
+		EGraphConnectionType GetType() const { return Type; }
+		FGuid GetGuid() const { return Guid; }
+
+		virtual bool AddConnection(FConnectionTypeBase* In) { return false; };
+		virtual bool RemoveConnection(FConnectionTypeBase* In) { return false; }
+
+		virtual TArray< FConnectionTypeBase* > GetBaseInputs() { return TArray<FConnectionTypeBase* >(); }
+		virtual TArray< FConnectionTypeBase* > GetBaseOutputs() { return TArray<FConnectionTypeBase* >(); }
+
+	};
+
+	template<typename T> class FOutput;
+
+	//
+	//  Input
+	//
+
+	template<class T>
+	struct EVALGRAPH_API FInputParameters {
+		FInputParameters(FString InName, FNode* InOwner, T InDefault = T())
+			: Name(InName)
+			, Owner(InOwner)
+			, Default(InDefault) {}
+		FString Name;
+		FNode* Owner = nullptr;
+		T Default;
+	};
+
+	template<class T>
+	class EVALGRAPH_API FInput : public FConnectionTypeBase
+	{
+		typedef FConnectionTypeBase Super;
+		FString Name;
+		T Default;
+		FOutput<T>* Connection;
+	public:
+		FInput(const FInputParameters<T>& Param, FGuid InGuid = FGuid::NewGuid());
+	
+		FString GetName() const { return Name; }
+		const T& GetDefault() const { return Default; }
+		
+		const FOutput<T>* GetConnection() const { return Connection; }
+		FOutput<T>* GetConnection() { return Connection; }
+
+		virtual bool AddConnection(FConnectionTypeBase* InOutput) override
+		{ 
+			ensure(Connection == nullptr);
+			if (ensure(InOutput->GetType()==this->GetType()))
+			{
+				Connection = (FOutput<T>*)InOutput;
+				return true;
+			}
+			return false;
+		}
+
+		virtual bool RemoveConnection(FConnectionTypeBase* InOutput) override
+		{ 
+			if (ensure(Connection == (FOutput<T>*)InOutput))
+			{
+				Connection = nullptr;
+				return true;
+			}
+			return false;
+		}
+
+		T GetValue(const FContext& Context)
+		{
+			if (Connection)
+			{
+				return Connection->Evaluate(Context);
+			}
+			return Default;
+		}
+
+		virtual TArray< FConnectionTypeBase* > GetBaseOutputs() override
+		{ 
+			TArray<FConnectionTypeBase* > RetList;
+			if (GetConnection())
+			{
+				RetList.Add(GetConnection());
+			}
+			return RetList;
+		}
+
+	};
+
+
+	//
+	// Output
+	//
+
+	struct EVALGRAPH_API FOutputParameters {
+		FOutputParameters(FString InName, FNode * InOwner) 
+			: Name(InName)
+			, Owner(InOwner) {}
+		FString Name;
+		FNode* Owner = nullptr;
+	};
+
+	template<class T>
+	class EVALGRAPH_API FOutput : public FConnectionTypeBase
+	{
+		typedef FConnectionTypeBase Super;
+
+		FString Name;
+		uint32 CacheKey = UINT_MAX;
+		TCacheValue<T> Cache;
+		TArray< FInput<T>* > Connections;
+	
+	public:
+		FOutput(const FOutputParameters& Param, FGuid InGuid = FGuid::NewGuid());
+
+		FString GetName() const { return Name; }
+
+		const TArray<FInput<T>*>& GetConnections() const { return Connections; }
+		TArray<FInput<T>* >& GetConnections() { return Connections; }
+
+		virtual TArray<FConnectionTypeBase*> GetBaseInputs() override
+		{ 
+			TArray<FConnectionTypeBase*> RetList;
+			RetList.Reserve(Connections.Num());
+			for (FInput<T>* Ptr : Connections) { RetList.Add(Ptr); }
+			return RetList;
+		}
+
+
+		virtual bool AddConnection(FConnectionTypeBase* InOutput) override
+		{
+			if (ensure(InOutput->GetType() == this->GetType()))
+			{
+				Connections.Add((FInput<T>*)InOutput);
+				return true;
+			}
+			return false;
+		}
+		
+		virtual bool RemoveConnection(FConnectionTypeBase* InInput) override { Connections.RemoveSwap((FInput<T>*)InInput); return true; }
+
+		void SetValue(T InVal, const FContext& Context);
+
+		T Evaluate(const FContext& Context);
+
+	};
+}
+
diff --git a/Engine/Plugins/Experimental/EvalGraph/Source/EvalGraph/Public/EvalGraph/EvalGraphNode.h b/Engine/Plugins/Experimental/EvalGraph/Source/EvalGraph/Public/EvalGraph/EvalGraphNode.h
new file mode 100644
--- /dev/null
+++ b/Engine/Plugins/Experimental/EvalGraph/Source/EvalGraph/Public/EvalGraph/EvalGraphNode.h
@@ -0,0 +1,56 @@
+// Copyright Epic Games, Inc. All Rights Reserved.
+
+#pragma once
+
+#include "CoreMinimal.h"
+#include "EvalGraph/EvalGraphInputOutput.h"
+#include "EvalGraph/EvalGraphNodeParameters.h"
+
+class FConnectionTypeBase;
+
+namespace Eg
+{
+
+	struct EVALGRAPH_API FNodeParameters {
+		FString Name;
+	};
+
+
+	//
+	// FNode
+	//
+	class FNode
+	{
+		friend class FGraph;
+		friend class FConnectionTypeBase;
+
+		FGuid Guid;
+		FName Name;
+		TArray< FConnectionTypeBase* > Inputs;
+		TArray< FConnectionTypeBase* > Outputs;
+	public:
+		FNode(const FNodeParameters& Param, FGuid InGuid = FGuid::NewGuid())
+			: Name(Param.Name) {}
+		virtual ~FNode() {}
+
+		FName GetName() const { return Name; }
+		void SetName(FName InName) { Name = InName; }
+
+
+
+		virtual void Evaluate(const FContext& Context, FConnectionTypeBase* ) { ensure(false); }
+
+	protected:
+		void AddBaseInput(FConnectionTypeBase* InPtr) { Inputs.Add(InPtr); }
+		const TArray< FConnectionTypeBase* >& GetInputs() const { return Inputs; }
+		TArray< FConnectionTypeBase* >& GetInputs() { return Inputs; }
+
+		void AddBaseOutput(FConnectionTypeBase* InPtr) { Outputs.Add(InPtr); }
+		const TArray< FConnectionTypeBase* >& GetOutputs() const { return Outputs; }
+		TArray< FConnectionTypeBase* >& GetOutputs() { return Outputs; }
+
+
+
+	};
+}
+
diff --git a/Engine/Plugins/Experimental/EvalGraph/Source/EvalGraph/Public/EvalGraph/EvalGraphNodeParameters.h b/Engine/Plugins/Experimental/EvalGraph/Source/EvalGraph/Public/EvalGraph/EvalGraphNodeParameters.h
new file mode 100644
--- /dev/null
+++ b/Engine/Plugins/Experimental/EvalGraph/Source/EvalGraph/Public/EvalGraph/EvalGraphNodeParameters.h
@@ -0,0 +1,41 @@
+// Copyright Epic Games, Inc. All Rights Reserved.
+
+#pragma once
+
+#include "CoreMinimal.h"
+#include "EvalGraph/EvalGraphConnectionTypes.h"
+
+class  UEvalGraph;
+class  UEdGraphPin;
+
+namespace Eg
+{
+	struct EVALGRAPH_API FContext 
+	{
+		float Timestamp = 0.f;
+
+		uint32 GetTypeHash() const
+		{
+			return ::GetTypeHash(Timestamp);
+		}
+	};
+
+
+	struct FCacheValueBase {
+		FCacheValueBase(EGraphConnectionType InType) : Type(InType) {}
+		EGraphConnectionType Type;
+	};
+
+	template<class T>
+	struct TCacheValue : public FCacheValueBase {
+		TCacheValue(T InData = T())
+			: FCacheValueBase(GraphConnectionType<T>())
+			, Data(InData) {}
+		T Data;
+	};
+}
+
+FORCEINLINE uint32 GetTypeHash(const Eg::FContext& Context)
+{
+	return ::GetTypeHash(Context.Timestamp);
+}
diff --git a/Engine/Plugins/Experimental/EvalGraph/Source/EvalGraph/Public/EvalGraph/EvalGraphPlugin.h b/Engine/Plugins/Experimental/EvalGraph/Source/EvalGraph/Public/EvalGraph/EvalGraphPlugin.h
new file mode 100644
--- /dev/null
+++ b/Engine/Plugins/Experimental/EvalGraph/Source/EvalGraph/Public/EvalGraph/EvalGraphPlugin.h
@@ -0,0 +1,39 @@
+// Copyright Epic Games, Inc. All Rights Reserved.
+
+#pragma once
+
+#include "CoreMinimal.h"
+#include "Modules/ModuleInterface.h"
+#include "Modules/ModuleManager.h"
+
+
+/**
+ * The public interface to this module
+ */
+class IEvalGraphPlugin : public IModuleInterface
+{
+
+public:
+
+	/**
+	 * Singleton-like access to this module's interface.  This is just for convenience!
+	 * Beware of calling this during the shutdown phase, though.  Your module might have been unloaded already.
+	 *
+	 * @return Returns singleton instance, loading the module on demand if needed
+	 */
+	static inline IEvalGraphPlugin& Get()
+	{
+		return FModuleManager::LoadModuleChecked< IEvalGraphPlugin >( "EvalGraphPlugin" );
+	}
+
+	/**
+	 * Checks to see if this module is loaded and ready.  It is only valid to call Get() if IsAvailable() returns true.
+	 *
+	 * @return True if the module is loaded and ready to use
+	 */
+	static inline bool IsAvailable()
+	{
+		return FModuleManager::Get().IsModuleLoaded( "EvalGraphPlugin" );
+	}
+};
+
