EvalGraph - Fix CIS non-unity

diff --git a/Engine/Plugins/Experimental/EvalGraph/Source/EvalGraph/Private/EvalGraph/EvalGraph.cpp b/Engine/Plugins/Experimental/EvalGraph/Source/EvalGraph/Private/EvalGraph/EvalGraph.cpp
--- a/Engine/Plugins/Experimental/EvalGraph/Source/EvalGraph/Private/EvalGraph/EvalGraph.cpp
+++ b/Engine/Plugins/Experimental/EvalGraph/Source/EvalGraph/Private/EvalGraph/EvalGraph.cpp
@@ -1,6 +1,5 @@
 // Copyright Epic Games, Inc. All Rights Reserved.
 
-#pragma once
 #include "EvalGraph/EvalGraph.h"
 
 #include "EvalGraph/EvalGraphInputOutput.h"
diff --git a/Engine/Plugins/Experimental/EvalGraph/Source/EvalGraph/Private/EvalGraph/EvalGraphInputOutput.cpp b/Engine/Plugins/Experimental/EvalGraph/Source/EvalGraph/Private/EvalGraph/EvalGraphInputOutput.cpp
--- a/Engine/Plugins/Experimental/EvalGraph/Source/EvalGraph/Private/EvalGraph/EvalGraphInputOutput.cpp
+++ b/Engine/Plugins/Experimental/EvalGraph/Source/EvalGraph/Private/EvalGraph/EvalGraphInputOutput.cpp
@@ -1,6 +1,5 @@
 // Copyright Epic Games, Inc. All Rights Reserved.
 
-#pragma once
 #include "EvalGraph/EvalGraphInputOutput.h"
 
 #include "EvalGraph/EvalGraphConnectionTypes.h"
diff --git a/Engine/Plugins/Experimental/EvalGraph/Source/EvalGraph/Private/EvalGraph/EvalGraphPluginPrivate.h b/Engine/Plugins/Experimental/EvalGraph/Source/EvalGraph/Private/EvalGraph/EvalGraphPluginPrivate.h
--- a/Engine/Plugins/Experimental/EvalGraph/Source/EvalGraph/Private/EvalGraph/EvalGraphPluginPrivate.h
+++ b/Engine/Plugins/Experimental/EvalGraph/Source/EvalGraph/Private/EvalGraph/EvalGraphPluginPrivate.h
@@ -1,7 +1,5 @@
 // Copyright Epic Games, Inc. All Rights Reserved.
 
-#pragma once
-
 #include "CoreMinimal.h"
 #include "Stats/Stats.h"
 
