EvalGraph - Serialization and factory construction, registration, and invalidation

diff --git a/Engine/Plugins/Experimental/EvalGraph/Source/EvalGraph/Private/EvalGraph/EvalGraph.cpp b/Engine/Plugins/Experimental/EvalGraph/Source/EvalGraph/Private/EvalGraph/EvalGraph.cpp
--- a/Engine/Plugins/Experimental/EvalGraph/Source/EvalGraph/Private/EvalGraph/EvalGraph.cpp
+++ b/Engine/Plugins/Experimental/EvalGraph/Source/EvalGraph/Private/EvalGraph/EvalGraph.cpp
@@ -5,7 +5,7 @@
 #include "EvalGraph/EvalGraphInputOutput.h"
 #include "EvalGraph/EvalGraphConnectionTypes.h"
 #include "EvalGraph/EvalGraphNodeParameters.h"
-#include "EvalGraph/EvalGraphNode.h"
+#include "EvalGraph/EvalGraphNodeFactory.h"
 
 namespace Eg
 {
@@ -14,6 +14,7 @@ namespace Eg
 		: Guid(InGuid)
 	{}
 
+
 	void FGraph::RemoveNode(TSharedPtr<FNode> Node)
 	{
 		for (FConnectionTypeBase* Output : Node->GetOutputs())
@@ -63,6 +64,102 @@ namespace Eg
 		Connections.RemoveSwap(FConnection(Input->GetGuid(), Output->GetGuid()));
 	}
 
+	void FGraph::Serialize(FArchive& Ar)
+	{
+		Ar << Guid;
+		if (Ar.IsSaving())
+		{
+			FGuid ArGuid;
+			FName ArType, ArName;
+			int32 ArNum = Nodes.Num();
+			Ar << ArNum;
+			for (TSharedPtr<FNode> Node : Nodes)
+			{
+				ArGuid = Node->GetGuid();
+				ArType = Node->GetType();
+				ArName = Node->GetName();
+				Ar << ArGuid << ArType << ArName;
+
+
+				TArray< FConnectionTypeBase* > IO = Node->GetOutputs();  IO.Append(Node->GetInputs());
+				ArNum = IO.Num();
+				Ar << ArNum;
+				for (FConnectionTypeBase* Conn : IO)
+				{
+					ArGuid = Conn->GetGuid();
+					ArName = Conn->GetName();
+					Ar << ArGuid << ArName;
+				}
+
+				Node->SerializeInternal(Ar);
+			}
+
+			Ar << Connections;
+
+		}
+		else if( Ar.IsLoading())
+		{
+			FGuid ArGuid;
+			FName ArType, ArName;
+			int32 ArNum = 0;
+			TMap<FGuid, FConnectionTypeBase* > ConnectionGuidMap;
+
+			Ar << ArNum;
+			for (int32 Ndx = ArNum; Ndx > 0; Ndx--)
+			{
+				Ar << ArGuid << ArType << ArName;
+				if (TSharedPtr<FNode> Node = FNodeFactory::GetInstance()->NewNodeFromRegisteredType(*this, { ArGuid,ArType,ArName }))
+				{
+					
+					//
+					// Build Indexing Map
+					//
+					TMap<FName, FConnectionTypeBase* > ConnectionNameMap;
+
+					TArray< FConnectionTypeBase* > IO = Node->GetOutputs();  IO.Append(Node->GetInputs());
+					for (FConnectionTypeBase* Conn : IO)
+					{
+						// @todo(eg) : This is a bit fragile, if the connection name changes then it breaks. 
+						ensure(!ConnectionNameMap.Contains(Conn->GetName()));
+						ConnectionNameMap.Add(Conn->GetName(), Conn);
+					}
+					
+					//
+					// Load and remap the nodes inputs and outputs. 
+					//
+					Ar << ArNum;
+					for (int Cdx=ArNum;Cdx>0; Cdx--)
+					{
+						Ar << ArGuid << ArName;
+
+						ensure(ConnectionNameMap.Contains(ArName));
+						ConnectionNameMap[ArName]->SetGuid(ArGuid);
+
+						ensure(!ConnectionGuidMap.Contains(ArGuid));
+						ConnectionGuidMap.Add(ArGuid, ConnectionNameMap[ArName]);
+					}
+
+					Node->SerializeInternal(Ar);
+				}
+				else
+				{
+					ensureMsgf(false, TEXT("Warning: Missing registered node type on load. Graph will fail to evaluate (%s %s)"), *ArType.ToString(), *ArName.ToString());
+				}
+				
+			}
+
+			TArray< FConnection > LocalConnections;
+			Ar << LocalConnections;
+			for (const FConnection& Con : LocalConnections)
+			{
+				ensure(ConnectionGuidMap.Contains(Con.Key));
+				ensure(ConnectionGuidMap.Contains(Con.Value));
+				ensure(ConnectionGuidMap[Con.Key]->GetType() == ConnectionGuidMap[Con.Value]->GetType());
+				Connect(ConnectionGuidMap[Con.Key], ConnectionGuidMap[Con.Value]);
+			}
+
+		}
+	}
 
 }
 
diff --git a/Engine/Plugins/Experimental/EvalGraph/Source/EvalGraph/Private/EvalGraph/EvalGraphInputOutput.cpp b/Engine/Plugins/Experimental/EvalGraph/Source/EvalGraph/Private/EvalGraph/EvalGraphInputOutput.cpp
--- a/Engine/Plugins/Experimental/EvalGraph/Source/EvalGraph/Private/EvalGraph/EvalGraphInputOutput.cpp
+++ b/Engine/Plugins/Experimental/EvalGraph/Source/EvalGraph/Private/EvalGraph/EvalGraphInputOutput.cpp
@@ -9,8 +9,9 @@
 namespace Eg
 {
 
-	FConnectionTypeBase::FConnectionTypeBase(EGraphConnectionType InType, FNode* InOwningNode, FGuid InGuid)
+	FConnectionTypeBase::FConnectionTypeBase(EGraphConnectionType InType, FName InName, FNode* InOwningNode, FGuid InGuid)
 		: Type(InType)
+		, Name(InName)
 		, Guid(InGuid)
 		, OwningNode(InOwningNode)
 	{}
@@ -23,19 +24,32 @@ namespace Eg
 
 	template<class T>
 	FInput<T>::FInput(const FInputParameters<T>& Param, FGuid InGuid )
-		: FConnectionTypeBase(GraphConnectionType<T>(), Param.Owner, InGuid)
-		, Name(Param.Name)
+		: FConnectionTypeBase(GraphConnectionType<T>(), Param.Name, Param.Owner, InGuid)
 		, Default(Param.Default)
 		, Connection(nullptr)
 	{
 		Super::AddBaseInput(Param.Owner, this);
 	}
 
+	template<class T>
+	void FInput<T>::SetValue(const T& Value, const FContext& Context)
+	{
+		Default = Value;
+		if (!Connection)
+		{
+			OwningNode->InvalidateOutputs();
+		}
+	}
+
+	template<class T>
+	void FInput<T>::Invalidate()
+	{
+		OwningNode->InvalidateOutputs();
+	}
 
 	template<class T>
 	FOutput<T>::FOutput(const FOutputParameters& Param, FGuid InGuid )
-		: FConnectionTypeBase(GraphConnectionType<T>(), Param.Owner, InGuid)
-		, Name(Param.Name)
+		: FConnectionTypeBase(GraphConnectionType<T>(), Param.Name, Param.Owner, InGuid)
 	{
 		Super::AddBaseOutput(Param.Owner, this);
 	}
@@ -59,6 +73,21 @@ namespace Eg
 		return Cache.Data;
 	}
 
+	template<class T>
+	void FOutput<T>::Invalidate( )
+	{
+		if (CacheKey != UINT_MAX)
+		{
+			CacheKey = UINT_MAX;
+			Cache = T();
+
+			for (FConnectionTypeBase* Con : GetConnections())
+			{
+				Con->Invalidate();
+			}
+		}
+	}
+
 
 #define EVAL_GRAPH_CONNECTION_TYPE(a,A) template class FOutput<a>;
 #include "../../Public/EvalGraph/EvalGraphConnectionTypeValues.inl"
diff --git a/Engine/Plugins/Experimental/EvalGraph/Source/EvalGraph/Private/EvalGraph/EvalGraphNode.cpp b/Engine/Plugins/Experimental/EvalGraph/Source/EvalGraph/Private/EvalGraph/EvalGraphNode.cpp
--- a/Engine/Plugins/Experimental/EvalGraph/Source/EvalGraph/Private/EvalGraph/EvalGraphNode.cpp
+++ b/Engine/Plugins/Experimental/EvalGraph/Source/EvalGraph/Private/EvalGraph/EvalGraphNode.cpp
@@ -6,6 +6,13 @@
 
 namespace Eg
 {
+	void FNode::InvalidateOutputs()
+	{
+		for(FConnectionTypeBase * Output :  Outputs)
+		{
+			Output->Invalidate();
+		}
+	}
 
 }
 
diff --git a/Engine/Plugins/Experimental/EvalGraph/Source/EvalGraph/Private/EvalGraph/EvalGraphNodeFactory.cpp b/Engine/Plugins/Experimental/EvalGraph/Source/EvalGraph/Private/EvalGraph/EvalGraphNodeFactory.cpp
new file mode 100644
--- /dev/null
+++ b/Engine/Plugins/Experimental/EvalGraph/Source/EvalGraph/Private/EvalGraph/EvalGraphNodeFactory.cpp
@@ -0,0 +1,21 @@
+// Copyright Epic Games, Inc. All Rights Reserved.
+
+#include "EvalGraph/EvalGraphNodeFactory.h"
+
+#include "EvalGraph/EvalGraphConnectionTypes.h"
+#include "EvalGraph/EvalGraphNode.h"
+
+namespace Eg
+{
+	FNodeFactory* FNodeFactory::Instance = nullptr;
+
+	TSharedPtr<FNode> FNodeFactory::NewNodeFromRegisteredType(FGraph& Graph, const FNewNodeParameters& Param)
+	{ 
+		if (ClassMap.Contains(Param.Type))
+		{
+			return Graph.AddNode(ClassMap[Param.Type](Param));
+		}
+		return TSharedPtr<FNode>(nullptr);
+	}
+}
+
diff --git a/Engine/Plugins/Experimental/EvalGraph/Source/EvalGraph/Public/EvalGraph/EvalGraph.h b/Engine/Plugins/Experimental/EvalGraph/Source/EvalGraph/Public/EvalGraph/EvalGraph.h
--- a/Engine/Plugins/Experimental/EvalGraph/Source/EvalGraph/Public/EvalGraph/EvalGraph.h
+++ b/Engine/Plugins/Experimental/EvalGraph/Source/EvalGraph/Public/EvalGraph/EvalGraph.h
@@ -3,14 +3,15 @@
 #pragma once
 
 #include "CoreMinimal.h"
-
+#include "Chaos/ChaosArchive.h"
+#include "Serialization/Archive.h"
+#include "EvalGraph/EvalGraphNode.h"
 
 namespace Eg
 {
-	class FNode;
 	class FConnectionTypeBase;
 
-	typedef TTuple<FGuid, FGuid> FConnection;
+	typedef TTuple<FGuid, FGuid> FConnection; // <Input, Output>
 
 
 	//
@@ -22,10 +23,11 @@ namespace Eg
 		FGuid  Guid;
 		TArray< TSharedPtr<FNode> > Nodes;
 		TArray< FConnection > Connections;
-		TMap<FGuid, TSharedPtr<FConnectionTypeBase> > ConnectionMap;
+		//TMap<FGuid, TSharedPtr<FConnectionTypeBase> > ConnectionMap;
 
 	public:
 		FGraph(FGuid InGuid = FGuid::NewGuid());
+		virtual ~FGraph() {}
 
 		int NumNodes() { return Nodes.Num(); }
 
@@ -36,11 +38,43 @@ namespace Eg
 			return NewNode;
 		}
 
+		template<class T> T* FindNode(FName InName)
+		{
+			for (TSharedPtr<FNode> Node : Nodes)  
+			{
+				if (Node->GetName().ToString().Equals(InName.ToString())) 
+				{
+					if (Node->GetType().ToString().Equals(T::Type.ToString())) 
+					{
+						return reinterpret_cast<T*>(Node.Get()); // @todo(eg) : Can we do better here?
+					}
+				}
+			}
+			return nullptr;
+		}
+
+
 		void RemoveNode(TSharedPtr<FNode> Node);
 
 		void Connect(FConnectionTypeBase* Input, FConnectionTypeBase* Output);
 		void Disconnect(FConnectionTypeBase* Input, FConnectionTypeBase* Output);
+
+		virtual void Serialize(FArchive& Ar);
+
 	};
 
 }
 
+FORCEINLINE FArchive& operator<<(FArchive& Ar, Eg::FGraph& Value)
+{
+	Value.Serialize(Ar);
+	return Ar;
+}
+
+FORCEINLINE FArchive& operator<<(Chaos::FChaosArchive& Ar, Eg::FGraph& Value)
+{
+	Value.Serialize(Ar);
+	return Ar;
+}
+
+
diff --git a/Engine/Plugins/Experimental/EvalGraph/Source/EvalGraph/Public/EvalGraph/EvalGraphInputOutput.h b/Engine/Plugins/Experimental/EvalGraph/Source/EvalGraph/Public/EvalGraph/EvalGraphInputOutput.h
--- a/Engine/Plugins/Experimental/EvalGraph/Source/EvalGraph/Public/EvalGraph/EvalGraphInputOutput.h
+++ b/Engine/Plugins/Experimental/EvalGraph/Source/EvalGraph/Public/EvalGraph/EvalGraphInputOutput.h
@@ -18,6 +18,7 @@ namespace Eg
 
 	protected:
 		EGraphConnectionType Type;
+		FName Name;
 		FGuid  Guid;
 		FNode* OwningNode = nullptr;
 
@@ -25,18 +26,26 @@ namespace Eg
 		static void AddBaseInput(FNode* InNode, FConnectionTypeBase*);
 		static void AddBaseOutput(FNode* InNode, FConnectionTypeBase*);
 
+
 	public:
-		FConnectionTypeBase(EGraphConnectionType InType, FNode* OwningNode = nullptr, FGuid InGuid = FGuid::NewGuid());
+		FConnectionTypeBase(EGraphConnectionType InType, FName InName, FNode* OwningNode = nullptr, FGuid InGuid = FGuid::NewGuid());
 		virtual ~FConnectionTypeBase() {};
 
 		EGraphConnectionType GetType() const { return Type; }
+
 		FGuid GetGuid() const { return Guid; }
+		void SetGuid(FGuid InGuid) { Guid = InGuid; }
+
+		FName GetName() const { return Name; }
+		void SetName(FName InName) { Name = InName; }
+
 
 		virtual bool AddConnection(FConnectionTypeBase* In) { return false; };
 		virtual bool RemoveConnection(FConnectionTypeBase* In) { return false; }
 
 		virtual TArray< FConnectionTypeBase* > GetBaseInputs() { return TArray<FConnectionTypeBase* >(); }
 		virtual TArray< FConnectionTypeBase* > GetBaseOutputs() { return TArray<FConnectionTypeBase* >(); }
+		virtual void Invalidate() {};
 
 	};
 
@@ -48,11 +57,11 @@ namespace Eg
 
 	template<class T>
 	struct EVALGRAPH_API FInputParameters {
-		FInputParameters(FString InName, FNode* InOwner, T InDefault = T())
+		FInputParameters(FName InName, FNode* InOwner, T InDefault = T())
 			: Name(InName)
 			, Owner(InOwner)
 			, Default(InDefault) {}
-		FString Name;
+		FName Name;
 		FNode* Owner = nullptr;
 		T Default;
 	};
@@ -60,14 +69,14 @@ namespace Eg
 	template<class T>
 	class EVALGRAPH_API FInput : public FConnectionTypeBase
 	{
+		friend class FConnectionTypeBase;
+
 		typedef FConnectionTypeBase Super;
-		FString Name;
 		T Default;
 		FOutput<T>* Connection;
 	public:
 		FInput(const FInputParameters<T>& Param, FGuid InGuid = FGuid::NewGuid());
 	
-		FString GetName() const { return Name; }
 		const T& GetDefault() const { return Default; }
 		
 		const FOutput<T>* GetConnection() const { return Connection; }
@@ -103,6 +112,8 @@ namespace Eg
 			return Default;
 		}
 
+		void SetValue(const T& Value, const FContext& Context);
+
 		virtual TArray< FConnectionTypeBase* > GetBaseOutputs() override
 		{ 
 			TArray<FConnectionTypeBase* > RetList;
@@ -113,6 +124,8 @@ namespace Eg
 			return RetList;
 		}
 
+		virtual void Invalidate() override;
+
 	};
 
 
@@ -121,10 +134,10 @@ namespace Eg
 	//
 
 	struct EVALGRAPH_API FOutputParameters {
-		FOutputParameters(FString InName, FNode * InOwner) 
+		FOutputParameters(FName InName, FNode * InOwner) 
 			: Name(InName)
 			, Owner(InOwner) {}
-		FString Name;
+		FName Name;
 		FNode* Owner = nullptr;
 	};
 
@@ -132,8 +145,8 @@ namespace Eg
 	class EVALGRAPH_API FOutput : public FConnectionTypeBase
 	{
 		typedef FConnectionTypeBase Super;
+		friend class FConnectionTypeBase;
 
-		FString Name;
 		uint32 CacheKey = UINT_MAX;
 		TCacheValue<T> Cache;
 		TArray< FInput<T>* > Connections;
@@ -141,8 +154,6 @@ namespace Eg
 	public:
 		FOutput(const FOutputParameters& Param, FGuid InGuid = FGuid::NewGuid());
 
-		FString GetName() const { return Name; }
-
 		const TArray<FInput<T>*>& GetConnections() const { return Connections; }
 		TArray<FInput<T>* >& GetConnections() { return Connections; }
 
@@ -171,6 +182,8 @@ namespace Eg
 
 		T Evaluate(const FContext& Context);
 
+		virtual void Invalidate() override;
+
 	};
 }
 
diff --git a/Engine/Plugins/Experimental/EvalGraph/Source/EvalGraph/Public/EvalGraph/EvalGraphNode.h b/Engine/Plugins/Experimental/EvalGraph/Source/EvalGraph/Public/EvalGraph/EvalGraphNode.h
--- a/Engine/Plugins/Experimental/EvalGraph/Source/EvalGraph/Public/EvalGraph/EvalGraphNode.h
+++ b/Engine/Plugins/Experimental/EvalGraph/Source/EvalGraph/Public/EvalGraph/EvalGraphNode.h
@@ -12,7 +12,7 @@ namespace Eg
 {
 
 	struct EVALGRAPH_API FNodeParameters {
-		FString Name;
+		FName Name;
 	};
 
 
@@ -26,19 +26,29 @@ namespace Eg
 
 		FGuid Guid;
 		FName Name;
+
 		TArray< FConnectionTypeBase* > Inputs;
 		TArray< FConnectionTypeBase* > Outputs;
 	public:
+
 		FNode(const FNodeParameters& Param, FGuid InGuid = FGuid::NewGuid())
-			: Name(Param.Name) {}
+			: Guid(InGuid)
+			, Name(Param.Name)
+		{}
+
 		virtual ~FNode() {}
 
 		FName GetName() const { return Name; }
 		void SetName(FName InName) { Name = InName; }
 
+		virtual FName GetType() const { check(true); return FName("invalid"); }
+		FGuid GetGuid() const { return Guid; }
 
 
-		virtual void Evaluate(const FContext& Context, FConnectionTypeBase* ) { ensure(false); }
+		virtual void Evaluate(const FContext& Context, FConnectionTypeBase*) { ensure(false); }
+		void InvalidateOutputs();
+
+		virtual void SerializeInternal(FArchive& Ar) {};
 
 	protected:
 		void AddBaseInput(FConnectionTypeBase* InPtr) { Inputs.Add(InPtr); }
@@ -49,8 +59,6 @@ namespace Eg
 		const TArray< FConnectionTypeBase* >& GetOutputs() const { return Outputs; }
 		TArray< FConnectionTypeBase* >& GetOutputs() { return Outputs; }
 
-
-
 	};
 }
 
diff --git a/Engine/Plugins/Experimental/EvalGraph/Source/EvalGraph/Public/EvalGraph/EvalGraphNodeFactory.h b/Engine/Plugins/Experimental/EvalGraph/Source/EvalGraph/Public/EvalGraph/EvalGraphNodeFactory.h
new file mode 100644
--- /dev/null
+++ b/Engine/Plugins/Experimental/EvalGraph/Source/EvalGraph/Public/EvalGraph/EvalGraphNodeFactory.h
@@ -0,0 +1,59 @@
+// Copyright Epic Games, Inc. All Rights Reserved.
+
+#pragma once
+
+#include "CoreMinimal.h"
+#include "EvalGraph/EvalGraph.h"
+
+namespace Eg
+{
+	class FNode;
+	class FConnectionTypeBase;
+
+	struct EVALGRAPH_API FNewNodeParameters {
+		FGuid Guid;
+		FName Type;
+		FName Name; 
+	};
+
+	//
+	//
+	//
+	class EVALGRAPH_API FNodeFactory
+	{
+		typedef TFunction<FNode* (const FNewNodeParameters&)> FNewNodeFunction;
+
+		TMap<FName, FNewNodeFunction > ClassMap;
+
+		static FNodeFactory* Instance;
+		FNodeFactory() {}
+
+	public:
+		~FNodeFactory() { delete Instance; }
+
+		static FNodeFactory* GetInstance()
+		{
+			if (!Instance)
+			{
+				Instance = new FNodeFactory();
+			}
+			return Instance;
+		}
+
+		void RegisterNode(const FName& Type, FNewNodeFunction NewFunction)
+		{
+			ensure(!ClassMap.Contains(Type));
+			ClassMap.Add(Type, NewFunction);
+		}
+
+		TSharedPtr<FNode> NewNodeFromRegisteredType(FGraph& Graph, const FNewNodeParameters& Param);
+
+		void here() {}
+		template<class T> TSharedPtr<T> NewNode(FGraph& Graph, const FNewNodeParameters& Param)
+		{
+			return Graph.AddNode(new T(Param.Name, Param.Guid));
+		}
+	};
+
+}
+
