EvalGraph - Serialization updates

diff --git a/Engine/Plugins/Experimental/EvalGraph/Source/EvalGraph/Private/EvalGraph/EvalGraph.cpp b/Engine/Plugins/Experimental/EvalGraph/Source/EvalGraph/Private/EvalGraph/EvalGraph.cpp
--- a/Engine/Plugins/Experimental/EvalGraph/Source/EvalGraph/Private/EvalGraph/EvalGraph.cpp
+++ b/Engine/Plugins/Experimental/EvalGraph/Source/EvalGraph/Private/EvalGraph/EvalGraph.cpp
@@ -6,6 +6,9 @@
 #include "EvalGraph/EvalGraphConnectionTypes.h"
 #include "EvalGraph/EvalGraphNodeParameters.h"
 #include "EvalGraph/EvalGraphNodeFactory.h"
+#include "Logging/LogMacros.h"
+
+DEFINE_LOG_CATEGORY_STATIC(EGGRAPH_LOG, Error, All);
 
 namespace Eg
 {
@@ -71,6 +74,7 @@ namespace Eg
 		{
 			FGuid ArGuid;
 			FName ArType, ArName;
+			uint32 ArIntType;
 			int32 ArNum = Nodes.Num();
 			Ar << ArNum;
 			for (TSharedPtr<FNode> Node : Nodes)
@@ -87,8 +91,9 @@ namespace Eg
 				for (FConnectionTypeBase* Conn : IO)
 				{
 					ArGuid = Conn->GetGuid();
+					ArIntType = (int32)Conn->GetType();
 					ArName = Conn->GetName();
-					Ar << ArGuid << ArName;
+					Ar << ArGuid << ArIntType << ArName;
 				}
 
 				Node->SerializeInternal(Ar);
@@ -101,7 +106,9 @@ namespace Eg
 		{
 			FGuid ArGuid;
 			FName ArType, ArName;
+			uint32 ArIntType;
 			int32 ArNum = 0;
+
 			TMap<FGuid, FConnectionTypeBase* > ConnectionGuidMap;
 
 			Ar << ArNum;
@@ -110,33 +117,26 @@ namespace Eg
 				Ar << ArGuid << ArType << ArName;
 				if (TSharedPtr<FNode> Node = FNodeFactory::GetInstance()->NewNodeFromRegisteredType(*this, { ArGuid,ArType,ArName }))
 				{
-					
-					//
-					// Build Indexing Map
-					//
-					TMap<FName, FConnectionTypeBase* > ConnectionNameMap;
-
+					int ArNumInner;
+					Ar << ArNumInner;
 					TArray< FConnectionTypeBase* > IO = Node->GetOutputs();  IO.Append(Node->GetInputs());
-					for (FConnectionTypeBase* Conn : IO)
+					for (int Cdx=0;Cdx< ArNumInner; Cdx++)
 					{
-						// @todo(eg) : This is a bit fragile, if the connection name changes then it breaks. 
-						ensure(!ConnectionNameMap.Contains(Conn->GetName()));
-						ConnectionNameMap.Add(Conn->GetName(), Conn);
-					}
-					
-					//
-					// Load and remap the nodes inputs and outputs. 
-					//
-					Ar << ArNum;
-					for (int Cdx=ArNum;Cdx>0; Cdx--)
-					{
-						Ar << ArGuid << ArName;
-
-						ensure(ConnectionNameMap.Contains(ArName));
-						ConnectionNameMap[ArName]->SetGuid(ArGuid);
-
-						ensure(!ConnectionGuidMap.Contains(ArGuid));
-						ConnectionGuidMap.Add(ArGuid, ConnectionNameMap[ArName]);
+						Ar << ArGuid << ArIntType << ArName;
+						if (Cdx < IO.Num())
+						{
+							if ((EGraphConnectionType)ArIntType!=IO[Cdx]->GetType())
+							{
+								FName TypeName = GraphConnectionTypeName((EGraphConnectionType)ArIntType);
+								FName FromType = GraphConnectionTypeName(IO[Cdx]->GetType());
+								UE_LOG(EGGRAPH_LOG, Fatal, TEXT("Type mismatch in input file for Ed::FGraphNode Input/Output (%s,%s) to type (%s)"),
+									*TypeName.ToString(), *ArName.ToString(), *FromType.ToString());
+							}
+
+							IO[Cdx]->SetGuid(ArGuid);
+							ensure(!ConnectionGuidMap.Contains(ArGuid));
+							ConnectionGuidMap.Add(ArGuid, IO[Cdx]);
+						}
 					}
 
 					Node->SerializeInternal(Ar);
diff --git a/Engine/Plugins/Experimental/EvalGraph/Source/EvalGraph/Public/EvalGraph/EvalGraphConnectionTypes.h b/Engine/Plugins/Experimental/EvalGraph/Source/EvalGraph/Public/EvalGraph/EvalGraphConnectionTypes.h
--- a/Engine/Plugins/Experimental/EvalGraph/Source/EvalGraph/Public/EvalGraph/EvalGraphConnectionTypes.h
+++ b/Engine/Plugins/Experimental/EvalGraph/Source/EvalGraph/Public/EvalGraph/EvalGraphConnectionTypes.h
@@ -12,7 +12,7 @@ namespace Eg
 	// see EvalGraphConnectionTypeValues.inl for specific types.
 	//
 #define EVAL_GRAPH_CONNECTION_TYPE(a,A) F##A##Type,
-	enum class EGraphConnectionType : uint8
+	enum class EGraphConnectionType : uint32
 	{
 		FNoneType,
 #include "EvalGraphConnectionTypeValues.inl"
@@ -37,6 +37,22 @@ namespace Eg
 #include "EvalGraphConnectionTypeValues.inl"
 #undef EVAL_GRAPH_CONNECTION_TYPE
 
+// ---------------------------------------------------------
+//  GraphConnectionTypeName
+//    Templated function to return a EGraphConnectionType as a FName
+//
+	inline FName GraphConnectionTypeName(EGraphConnectionType ValueType)
+	{
+		switch (ValueType)
+		{
+#define EVAL_GRAPH_CONNECTION_TYPE(a,A)	case EGraphConnectionType::F##A##Type:\
+		return FName(TEXT(#A));
+#include "EvalGraphConnectionTypeValues.inl"
+#undef EVAL_GRAPH_CONNECTION_TYPE
+		}
+		return FName("FNoneType");
+	}
+
 	// ---------------------------------------------------------
 	//  void*
 	//     Returns a new EGraphConnectionType pointer based on 
