EvalGraph - Engine level.

diff --git a/Engine/Plugins/Experimental/EvalGraph/EvalGraph.uplugin b/Engine/Plugins/Experimental/EvalGraph/EvalGraph.uplugin
--- a/Engine/Plugins/Experimental/EvalGraph/EvalGraph.uplugin
+++ b/Engine/Plugins/Experimental/EvalGraph/EvalGraph.uplugin
@@ -21,6 +21,11 @@
 			"Type": "Runtime",
 			"LoadingPhase": "PreDefault"
 		},
+		{
+			"Name": "EvalGraphEngine",
+			"Type": "Runtime",
+			"LoadingPhase": "Default"
+		}
 	],
 	"Plugins": [
 	]
diff --git a/Engine/Plugins/Experimental/EvalGraph/Source/EvalGraphEngine/EvalGraphEngine.Build.cs b/Engine/Plugins/Experimental/EvalGraph/Source/EvalGraphEngine/EvalGraphEngine.Build.cs
new file mode 100644
--- /dev/null
+++ b/Engine/Plugins/Experimental/EvalGraph/Source/EvalGraphEngine/EvalGraphEngine.Build.cs
@@ -0,0 +1,24 @@
+// Copyright Epic Games, Inc. All Rights Reserved.
+
+namespace UnrealBuildTool.Rules
+{
+	public class EvalGraphEngine : ModuleRules
+	{
+        public EvalGraphEngine(ReadOnlyTargetRules Target) : base(Target)
+		{
+			PrivateIncludePaths.Add("EvalGraphEngine/Private");
+            PublicIncludePaths.Add(ModuleDirectory + "/Public");
+
+			PublicDependencyModuleNames.AddRange(
+				new string[]
+				{
+					"Core",
+					"CoreUObject",
+					"Engine",
+					"SlateCore",
+					"EvalGraph"
+				}
+				);
+		}
+	}
+}
diff --git a/Engine/Plugins/Experimental/EvalGraph/Source/EvalGraphEngine/Private/EvalGraph/EvalGraphEnginePlugin.cpp b/Engine/Plugins/Experimental/EvalGraph/Source/EvalGraphEngine/Private/EvalGraph/EvalGraphEnginePlugin.cpp
new file mode 100644
--- /dev/null
+++ b/Engine/Plugins/Experimental/EvalGraph/Source/EvalGraphEngine/Private/EvalGraph/EvalGraphEnginePlugin.cpp
@@ -0,0 +1,32 @@
+// Copyright Epic Games, Inc. All Rights Reserved.
+
+#include "EvalGraph/EvalGraphEnginePlugin.h"
+
+#include "CoreMinimal.h"
+#include "Modules/ModuleManager.h"
+
+
+
+class FEvalGraphEnginePlugin : public IEvalGraphEnginePlugin
+{
+	/** IModuleInterface implementation */
+	virtual void StartupModule() override;
+	virtual void ShutdownModule() override;
+};
+
+IMPLEMENT_MODULE( FEvalGraphEnginePlugin, EvalGraphEngine )
+
+
+void FEvalGraphEnginePlugin::StartupModule()
+{
+	
+}
+
+
+void FEvalGraphEnginePlugin::ShutdownModule()
+{
+	
+}
+
+
+
diff --git a/Engine/Plugins/Experimental/EvalGraph/Source/EvalGraphEngine/Private/EvalGraph/EvalGraphEnginePluginPrivate.h b/Engine/Plugins/Experimental/EvalGraph/Source/EvalGraphEngine/Private/EvalGraph/EvalGraphEnginePluginPrivate.h
new file mode 100644
--- /dev/null
+++ b/Engine/Plugins/Experimental/EvalGraph/Source/EvalGraphEngine/Private/EvalGraph/EvalGraphEnginePluginPrivate.h
@@ -0,0 +1,11 @@
+// Copyright Epic Games, Inc. All Rights Reserved.
+
+#pragma once
+
+#include "CoreMinimal.h"
+#include "Stats/Stats.h"
+
+// You should place include statements to your module's private header files here.  You only need to
+// add includes for headers that are used in most of your module's source files though.
+
+
diff --git a/Engine/Plugins/Experimental/EvalGraph/Source/EvalGraphEngine/Private/EvalGraph/EvalGraphObject.cpp b/Engine/Plugins/Experimental/EvalGraph/Source/EvalGraphEngine/Private/EvalGraph/EvalGraphObject.cpp
new file mode 100644
--- /dev/null
+++ b/Engine/Plugins/Experimental/EvalGraph/Source/EvalGraphEngine/Private/EvalGraph/EvalGraphObject.cpp
@@ -0,0 +1,33 @@
+// Copyright Epic Games, Inc. All Rights Reserved.
+
+#include "EvalGraph/EvalGraphObject.h"
+#include "EvalGraph/EvalGraph.h"
+
+
+#define LOCTEXT_NAMESPACE "UEvalGraph")
+
+UEvalGraph::UEvalGraph(const FObjectInitializer& ObjectInitializer)
+	: UObject(ObjectInitializer)
+	, EvalGraph(new Eg::FGraph())
+{}
+
+#if WITH_EDITOR
+
+void UEvalGraph::PostEditChangeProperty(struct FPropertyChangedEvent& PropertyChangedEvent)
+{
+	UObject::PostEditChangeProperty(PropertyChangedEvent);
+}
+
+#endif
+
+void UEvalGraph::PostLoad()
+{
+	UObject::PostLoad();
+}
+
+void UEvalGraph::Serialize(FArchive& Ar)
+{
+	Ar << *EvalGraph.Get();
+}
+
+#undef LOCTEXT_NAMESPACE
diff --git a/Engine/Plugins/Experimental/EvalGraph/Source/EvalGraphEngine/Public/EvalGraph/EvalGraphEnginePlugin.h b/Engine/Plugins/Experimental/EvalGraph/Source/EvalGraphEngine/Public/EvalGraph/EvalGraphEnginePlugin.h
new file mode 100644
--- /dev/null
+++ b/Engine/Plugins/Experimental/EvalGraph/Source/EvalGraphEngine/Public/EvalGraph/EvalGraphEnginePlugin.h
@@ -0,0 +1,39 @@
+// Copyright Epic Games, Inc. All Rights Reserved.
+
+#pragma once
+
+#include "CoreMinimal.h"
+#include "Modules/ModuleInterface.h"
+#include "Modules/ModuleManager.h"
+
+
+/**
+ * The public interface to this module
+ */
+class IEvalGraphEnginePlugin : public IModuleInterface
+{
+
+public:
+
+	/**
+	 * Singleton-like access to this module's interface.  This is just for convenience!
+	 * Beware of calling this during the shutdown phase, though.  Your module might have been unloaded already.
+	 *
+	 * @return Returns singleton instance, loading the module on demand if needed
+	 */
+	static inline IEvalGraphEnginePlugin& Get()
+	{
+		return FModuleManager::LoadModuleChecked< IEvalGraphEnginePlugin >( "EvalGraphEnginePlugin" );
+	}
+
+	/**
+	 * Checks to see if this module is loaded and ready.  It is only valid to call Get() if IsAvailable() returns true.
+	 *
+	 * @return True if the module is loaded and ready to use
+	 */
+	static inline bool IsAvailable()
+	{
+		return FModuleManager::Get().IsModuleLoaded( "EvalGraphEnginePlugin" );
+	}
+};
+
diff --git a/Engine/Plugins/Experimental/EvalGraph/Source/EvalGraphEngine/Public/EvalGraph/EvalGraphObject.h b/Engine/Plugins/Experimental/EvalGraph/Source/EvalGraphEngine/Public/EvalGraph/EvalGraphObject.h
new file mode 100644
--- /dev/null
+++ b/Engine/Plugins/Experimental/EvalGraph/Source/EvalGraphEngine/Public/EvalGraph/EvalGraphObject.h
@@ -0,0 +1,42 @@
+// Copyright Epic Games, Inc. All Rights Reserved.
+
+#pragma once
+
+#include "CoreMinimal.h"
+#include "EvalGraph/EvalGraph.h"
+
+#include "EvalGraphObject.generated.h"
+
+/**
+* UEvalGraph (UObject)
+*
+* UObject wrapper for the Eg::FGraph
+*
+*/
+UCLASS(BlueprintType, customconstructor)
+class EVALGRAPHENGINE_API UEvalGraph : public UObject
+{
+	GENERATED_UCLASS_BODY()
+
+	TSharedPtr<Eg::FGraph, ESPMode::ThreadSafe> EvalGraph;
+
+public:
+	UEvalGraph(const FObjectInitializer& ObjectInitializer = FObjectInitializer::Get());
+
+	/** UObject Interface */
+#if WITH_EDITOR
+	virtual void PostEditChangeProperty(struct FPropertyChangedEvent& PropertyChangedEvent) override;
+#endif
+	virtual void PostLoad() override;
+	/** End UObject Interface */
+
+	void Serialize(FArchive& Ar);
+
+	/** Accessors for internal geometry collection */
+	void SetEvalGraph(TSharedPtr<Eg::FGraph, ESPMode::ThreadSafe> EvalGraphIn) { EvalGraph = EvalGraphIn; }
+	TSharedPtr<Eg::FGraph, ESPMode::ThreadSafe>       GetEvalGraph() { return EvalGraph; }
+	const TSharedPtr<Eg::FGraph, ESPMode::ThreadSafe> GetEvalGraph() const { return EvalGraph; }
+
+
+};
+
