EvalGraph - Editior level

diff --git a/Engine/Plugins/Experimental/EvalGraph/EvalGraph.uplugin b/Engine/Plugins/Experimental/EvalGraph/EvalGraph.uplugin
--- a/Engine/Plugins/Experimental/EvalGraph/EvalGraph.uplugin
+++ b/Engine/Plugins/Experimental/EvalGraph/EvalGraph.uplugin
@@ -25,6 +25,11 @@
 			"Name": "EvalGraphEngine",
 			"Type": "Runtime",
 			"LoadingPhase": "Default"
+		},
+		{
+			"Name": "EvalGraphEditor",
+			"Type": "Editor",
+			"LoadingPhase": "Default"
 		}
 	],
 	"Plugins": [
diff --git a/Engine/Plugins/Experimental/EvalGraph/Source/EvalGraphEditor/EvalGraphEditor.Build.cs b/Engine/Plugins/Experimental/EvalGraph/Source/EvalGraphEditor/EvalGraphEditor.Build.cs
new file mode 100644
--- /dev/null
+++ b/Engine/Plugins/Experimental/EvalGraph/Source/EvalGraphEditor/EvalGraphEditor.Build.cs
@@ -0,0 +1,42 @@
+// Copyright Epic Games, Inc. All Rights Reserved.
+
+namespace UnrealBuildTool.Rules
+{
+	public class EvalGraphEditor : ModuleRules
+	{
+        public EvalGraphEditor(ReadOnlyTargetRules Target) : base(Target)
+		{
+			PrivateIncludePaths.Add("EvalGraphEditor/Private");
+            PublicIncludePaths.Add(ModuleDirectory + "/Public");
+
+			PublicDependencyModuleNames.AddRange(
+				new string[]
+				{
+					"Core",
+					"CoreUObject",
+				    "Slate",
+				    "SlateCore",
+				    "Engine",
+					"EditorFramework",
+					"UnrealEd",
+				    "PropertyEditor",
+				    "RenderCore",
+				    "RHI",
+				    "AssetTools",
+				    "AssetRegistry",
+				    "SceneOutliner",
+					"EditorStyle",
+					"AssetTools",
+					"ToolMenus",
+					"LevelEditor",
+					"InputCore",
+					"AdvancedPreviewScene",
+					"GraphEditor",
+					"EvalGraph",
+					"EvalGraphEngine",
+					"Slate",
+				}
+			);
+		}
+	}
+}
diff --git a/Engine/Plugins/Experimental/EvalGraph/Source/EvalGraphEditor/Private/EvalGraph/EvalGraphAssetActions.cpp b/Engine/Plugins/Experimental/EvalGraph/Source/EvalGraphEditor/Private/EvalGraph/EvalGraphAssetActions.cpp
new file mode 100644
--- /dev/null
+++ b/Engine/Plugins/Experimental/EvalGraph/Source/EvalGraphEditor/Private/EvalGraph/EvalGraphAssetActions.cpp
@@ -0,0 +1,45 @@
+// Copyright Epic Games, Inc. All Rights Reserved.
+
+#include "EvalGraph/EvalGraphAssetActions.h"
+
+#include "EvalGraph/EvalGraphObject.h"
+
+#define LOCTEXT_NAMESPACE "AssetActions_EvalGraphAsset"
+
+FText FEvalGraphAssetActions::GetName() const
+{
+	return LOCTEXT("Name", "Evalulation Graph");
+}
+
+UClass* FEvalGraphAssetActions::GetSupportedClass() const
+{
+	return UEvalGraph::StaticClass();
+}
+
+FColor FEvalGraphAssetActions::GetTypeColor() const
+{
+	return FColor(255, 127, 40);
+}
+
+void FEvalGraphAssetActions::GetActions(const TArray<UObject*>& InObjects,
+										class FMenuBuilder&     MenuBuilder)
+{
+}
+
+void FEvalGraphAssetActions::OpenAssetEditor(
+	const TArray<UObject*>& InObjects, TSharedPtr<IToolkitHost> EditWithinLevelEditor /*= TSharedPtr<IToolkitHost>()*/)
+{
+	FSimpleAssetEditor::CreateEditor(EToolkitMode::Standalone, EditWithinLevelEditor, InObjects);
+}
+
+uint32 FEvalGraphAssetActions::GetCategories()
+{
+	return EAssetTypeCategories::Physics;
+}
+
+FText FEvalGraphAssetActions::GetAssetDescription(const struct FAssetData& AssetData) const
+{
+	return LOCTEXT("Description", "An evulation graph for asset authoring.");
+}
+
+#undef LOCTEXT_NAMESPACE
diff --git a/Engine/Plugins/Experimental/EvalGraph/Source/EvalGraphEditor/Private/EvalGraph/EvalGraphAssetFactory.cpp b/Engine/Plugins/Experimental/EvalGraph/Source/EvalGraphEditor/Private/EvalGraph/EvalGraphAssetFactory.cpp
new file mode 100644
--- /dev/null
+++ b/Engine/Plugins/Experimental/EvalGraph/Source/EvalGraphEditor/Private/EvalGraph/EvalGraphAssetFactory.cpp
@@ -0,0 +1,36 @@
+// Copyright Epic Games, Inc. All Rights Reserved.
+
+#include "EvalGraph/EvalGraphAssetFactory.h"
+
+#include "EvalGraph/EvalGraphObject.h"
+
+UEvalGraphAssetFactory::UEvalGraphAssetFactory()
+{
+	SupportedClass = UEvalGraph::StaticClass();
+}
+
+bool UEvalGraphAssetFactory::CanCreateNew() const
+{
+	return true;
+}
+
+bool UEvalGraphAssetFactory::FactoryCanImport(const FString& Filename)
+{
+	return false;
+}
+
+UObject* UEvalGraphAssetFactory::FactoryCreateNew(UClass* InClass, UObject* InParent, FName InName, EObjectFlags Flags,
+												   UObject* Context, FFeedbackContext* Warn)
+{
+	return NewObject<UEvalGraph>(InParent, InClass, InName, Flags);
+}
+
+bool UEvalGraphAssetFactory::ShouldShowInNewMenu() const
+{
+	return true;
+}
+
+bool UEvalGraphAssetFactory::ConfigureProperties()
+{
+	return true;
+}
diff --git a/Engine/Plugins/Experimental/EvalGraph/Source/EvalGraphEditor/Private/EvalGraph/EvalGraphEditorActions.cpp b/Engine/Plugins/Experimental/EvalGraph/Source/EvalGraphEditor/Private/EvalGraph/EvalGraphEditorActions.cpp
new file mode 100644
--- /dev/null
+++ b/Engine/Plugins/Experimental/EvalGraph/Source/EvalGraphEditor/Private/EvalGraph/EvalGraphEditorActions.cpp
@@ -0,0 +1,27 @@
+// Copyright Epic Games, Inc. All Rights Reserved.
+
+#include "EvalGraph/EvalGraphEditorActions.h"
+
+#define LOCTEXT_NAMESPACE "EvalGraphEditorCommands"
+
+void FEvalGraphEditorCommandsImpl::RegisterCommands()
+{
+	UI_COMMAND(EvaluateNode, "Evaluate", "Trigger an evaluation of the selected node.", EUserInterfaceActionType::Button, FInputChord())
+}
+
+void FEvalGraphEditorCommands::Register()
+{
+	return FEvalGraphEditorCommandsImpl::Register();
+}
+
+const FEvalGraphEditorCommandsImpl& FEvalGraphEditorCommands::Get()
+{
+	return FEvalGraphEditorCommandsImpl::Get();
+}
+
+void FEvalGraphEditorCommands::Unregister()
+{
+	return FEvalGraphEditorCommandsImpl::Unregister();
+}
+
+#undef LOCTEXT_NAMESPACE
diff --git a/Engine/Plugins/Experimental/EvalGraph/Source/EvalGraphEditor/Private/EvalGraph/EvalGraphEditorPlugin.cpp b/Engine/Plugins/Experimental/EvalGraph/Source/EvalGraphEditor/Private/EvalGraph/EvalGraphEditorPlugin.cpp
new file mode 100644
--- /dev/null
+++ b/Engine/Plugins/Experimental/EvalGraph/Source/EvalGraphEditor/Private/EvalGraph/EvalGraphEditorPlugin.cpp
@@ -0,0 +1,45 @@
+// Copyright Epic Games, Inc. All Rights Reserved.
+
+#include "EvalGraph/EvalGraphEditorPlugin.h"
+
+#include "AssetToolsModule.h"
+#include "CoreMinimal.h"
+#include "EdGraphUtilities.h"
+#include "EvalGraph/EvalGraphAssetActions.h"
+#include "EvalGraph/EvalGraphNodeFactories.h"
+
+#define LOCTEXT_NAMESPACE "EvalGraphEditor"
+
+//#define BOX_BRUSH(StyleSet, RelativePath, ...) FSlateBoxBrush(StyleSet->RootToContentDir(RelativePath, TEXT(".png")), __VA_ARGS__)
+//#define DEFAULT_FONT(...) FCoreStyle::GetDefaultFontStyle(__VA_ARGS__)
+
+void IEvalGraphEditorPlugin::StartupModule()
+{
+	EvalGraphAssetActions = new FEvalGraphAssetActions();
+
+	FAssetToolsModule& AssetToolsModule = FAssetToolsModule::GetModule();
+	IAssetTools& AssetTools = AssetToolsModule.Get();
+	AssetTools.RegisterAssetTypeActions(MakeShareable(EvalGraphAssetActions));
+
+
+	//EvalGraphNodeFactory = MakeShareable(new FEvalGraphNodeFactory());
+	//FEdGraphUtilities::RegisterVisualNodeFactory(EvalGraphNodeFactory);
+}
+
+void IEvalGraphEditorPlugin::ShutdownModule()
+{
+	if (UObjectInitialized())
+	{
+		FAssetToolsModule& AssetToolsModule = FAssetToolsModule::GetModule();
+		IAssetTools& AssetTools = AssetToolsModule.Get();
+
+		AssetTools.UnregisterAssetTypeActions(EvalGraphAssetActions->AsShared());
+
+		//FEdGraphUtilities::UnregisterVisualNodeFactory(EvalGraphNodeFactory);
+	}
+}
+
+IMPLEMENT_MODULE(IEvalGraphEditorPlugin, EvalGraphEditor)
+
+
+#undef LOCTEXT_NAMESPACE
diff --git a/Engine/Plugins/Experimental/EvalGraph/Source/EvalGraphEditor/Private/EvalGraph/EvalGraphNodeFactories.cpp b/Engine/Plugins/Experimental/EvalGraph/Source/EvalGraphEditor/Private/EvalGraph/EvalGraphNodeFactories.cpp
new file mode 100644
--- /dev/null
+++ b/Engine/Plugins/Experimental/EvalGraph/Source/EvalGraphEditor/Private/EvalGraph/EvalGraphNodeFactories.cpp
@@ -0,0 +1,22 @@
+// Copyright Epic Games, Inc. All Rights Reserved.
+
+#include "EvalGraph/EvalGraphNodeFactories.h"
+
+#include "Logging/LogMacros.h"
+
+
+
+TSharedPtr<class SGraphNode> FEvalGraphNodeFactory::CreateNode(UEdGraphNode* InNode) const
+{
+	//
+	// @example(EvalGraphNode) Register SNode for display properties in editor. 
+	//
+	/**
+	if (UEvalGraphGraphNodeExample* Node = Cast<UEvalGraphGraphNodeExample>(InNode))
+	{
+		return SNew(SEvalGraphGraphNodeExample, Node);
+	}
+
+	*/
+	return NULL;
+}
diff --git a/Engine/Plugins/Experimental/EvalGraph/Source/EvalGraphEditor/Public/EvalGraph/EvalGraphAssetActions.h b/Engine/Plugins/Experimental/EvalGraph/Source/EvalGraphEditor/Public/EvalGraph/EvalGraphAssetActions.h
new file mode 100644
--- /dev/null
+++ b/Engine/Plugins/Experimental/EvalGraph/Source/EvalGraphEditor/Public/EvalGraph/EvalGraphAssetActions.h
@@ -0,0 +1,17 @@
+// Copyright Epic Games, Inc. All Rights Reserved.
+
+#pragma once
+
+#include "AssetTypeActions_Base.h"
+
+class FEvalGraphAssetActions : public FAssetTypeActions_Base
+{
+public:
+	virtual FText GetName() const override;
+	virtual UClass* GetSupportedClass() const override;
+	virtual FColor GetTypeColor() const override;
+	virtual void GetActions(const TArray<UObject*>& InObjects, class FMenuBuilder& MenuBuilder) override;
+	virtual void OpenAssetEditor(const TArray<UObject*>& InObjects, TSharedPtr<IToolkitHost> EditWithinLevelEditor = TSharedPtr<IToolkitHost>()) override;
+	virtual uint32 GetCategories() override;
+	virtual FText GetAssetDescription(const struct FAssetData& AssetData) const override;
+};
\ No newline at end of file
diff --git a/Engine/Plugins/Experimental/EvalGraph/Source/EvalGraphEditor/Public/EvalGraph/EvalGraphAssetFactory.h b/Engine/Plugins/Experimental/EvalGraph/Source/EvalGraphEditor/Public/EvalGraph/EvalGraphAssetFactory.h
new file mode 100644
--- /dev/null
+++ b/Engine/Plugins/Experimental/EvalGraph/Source/EvalGraphEditor/Public/EvalGraph/EvalGraphAssetFactory.h
@@ -0,0 +1,32 @@
+// Copyright Epic Games, Inc. All Rights Reserved.
+
+#pragma once
+
+#include "Factories/Factory.h"
+
+#include "EvalGraphAssetFactory.generated.h"
+
+UCLASS()
+class UEvalGraphAssetFactory : public UFactory
+{
+	GENERATED_BODY()
+public:
+
+	UEvalGraphAssetFactory();
+
+	/** UFactory Interface */
+	bool     CanCreateNew() const override;
+	bool     FactoryCanImport(const FString& Filename) override;
+	UObject* FactoryCreateNew(UClass*           InClass,
+							  UObject*          InParent,
+							  FName             InName,
+							  EObjectFlags      Flags,
+							  UObject*          Context,
+							  FFeedbackContext* Warn) override;
+
+	bool ShouldShowInNewMenu() const override;
+	bool ConfigureProperties() override;
+	/** End UFactory Interface */
+
+private:
+};
diff --git a/Engine/Plugins/Experimental/EvalGraph/Source/EvalGraphEditor/Public/EvalGraph/EvalGraphEditorActions.h b/Engine/Plugins/Experimental/EvalGraph/Source/EvalGraphEditor/Public/EvalGraph/EvalGraphEditorActions.h
new file mode 100644
--- /dev/null
+++ b/Engine/Plugins/Experimental/EvalGraph/Source/EvalGraphEditor/Public/EvalGraph/EvalGraphEditorActions.h
@@ -0,0 +1,35 @@
+// Copyright Epic Games, Inc. All Rights Reserved.
+
+#pragma once
+
+#include "CoreMinimal.h"
+#include "Framework/Commands/Commands.h"
+#include "EditorStyleSet.h"
+
+class FEvalGraphEditorCommandsImpl : public TCommands<FEvalGraphEditorCommandsImpl>
+{
+public:
+
+	FEvalGraphEditorCommandsImpl()
+		: TCommands<FEvalGraphEditorCommandsImpl>( TEXT("EvalGraphEditor"), NSLOCTEXT("Contexts", "EvalGraphEditor", "Scene Graph Editor"), NAME_None, FEditorStyle::GetStyleSetName() )
+	{
+	}	
+
+	virtual ~FEvalGraphEditorCommandsImpl()
+	{
+	}
+
+	EVALGRAPHEDITOR_API virtual void RegisterCommands() override;
+
+	TSharedPtr< FUICommandInfo > EvaluateNode;
+};
+
+class EVALGRAPHEDITOR_API FEvalGraphEditorCommands
+{
+public:
+	static void Register();
+
+	static const FEvalGraphEditorCommandsImpl& Get();
+
+	static void Unregister();
+};
diff --git a/Engine/Plugins/Experimental/EvalGraph/Source/EvalGraphEditor/Public/EvalGraph/EvalGraphEditorPlugin.h b/Engine/Plugins/Experimental/EvalGraph/Source/EvalGraphEditor/Public/EvalGraph/EvalGraphEditorPlugin.h
new file mode 100644
--- /dev/null
+++ b/Engine/Plugins/Experimental/EvalGraph/Source/EvalGraphEditor/Public/EvalGraph/EvalGraphEditorPlugin.h
@@ -0,0 +1,60 @@
+// Copyright Epic Games, Inc. All Rights Reserved.
+
+#pragma once
+
+#include "CoreMinimal.h"
+#include "Modules/ModuleInterface.h"
+#include "Modules/ModuleManager.h"
+//#include "Toolkits/IToolkit.h"
+
+//class ISlateStyle;
+//class FSlateStyleSet;
+//class FEvalGraphNodeFactory;
+class FEvalGraphAssetActions;
+
+/**
+ * The public interface to this module
+ */
+class IEvalGraphEditorPlugin : public IModuleInterface
+{
+public:
+	virtual void StartupModule();
+	virtual void ShutdownModule();
+
+	/**
+	 * Singleton-like access to this module's interface.  This is just for convenience!
+	 * Beware of calling this during the shutdown phase, though.  Your module might have been unloaded already.
+	 *
+	 * @return Returns singleton instance, loading the module on demand if needed
+	 */
+	static inline IEvalGraphEditorPlugin& Get()
+	{
+		return FModuleManager::LoadModuleChecked< IEvalGraphEditorPlugin >( "EvalGraphEditorPlugin" );
+	}
+
+	/**
+	 * Checks to see if this module is loaded and ready.  It is only valid to call Get() if IsAvailable() returns true.
+	 *
+	 * @return True if the module is loaded and ready to use
+	 */
+	static inline bool IsAvailable()
+	{
+		return FModuleManager::Get().IsModuleLoaded( "EvalGraphEditorPlugin" );
+	}
+
+	//TSharedPtr<FSlateStyleSet> GetStyleSet() { return StyleSet; }
+	//static FName GetEditorStyleName();
+	//static const ISlateStyle* GetEditorStyle();
+
+private:
+	//void RegisterMenus();
+
+private:
+
+	FEvalGraphAssetActions* EvalGraphAssetActions;
+
+	//TSharedPtr<FSlateStyleSet> StyleSet;
+	//TSharedPtr<FEvalGraphNodeFactory> EvalGraphNodeFactory;
+
+};
+
diff --git a/Engine/Plugins/Experimental/EvalGraph/Source/EvalGraphEditor/Public/EvalGraph/EvalGraphNodeFactories.h b/Engine/Plugins/Experimental/EvalGraph/Source/EvalGraphEditor/Public/EvalGraph/EvalGraphNodeFactories.h
new file mode 100644
--- /dev/null
+++ b/Engine/Plugins/Experimental/EvalGraph/Source/EvalGraphEditor/Public/EvalGraph/EvalGraphNodeFactories.h
@@ -0,0 +1,12 @@
+// Copyright Epic Games, Inc. All Rights Reserved.
+
+#pragma once
+
+#include "EdGraphUtilities.h"
+#include "SGraphNode.h"
+#include "EdGraph/EdGraphNode.h"
+
+class EVALGRAPHEDITOR_API FEvalGraphNodeFactory : public FGraphPanelNodeFactory
+{
+	virtual TSharedPtr<class SGraphNode> CreateNode(class UEdGraphNode* InNode) const override;
+};
