EvalGraph - BaseTypes

diff --git a/Engine/Plugins/Experimental/EvalGraph/Source/EvalGraph/Public/EvalGraph/EvalGraphConnectionTypeValues.inl b/Engine/Plugins/Experimental/EvalGraph/Source/EvalGraph/Public/EvalGraph/EvalGraphConnectionTypeValues.inl
--- a/Engine/Plugins/Experimental/EvalGraph/Source/EvalGraph/Public/EvalGraph/EvalGraphConnectionTypeValues.inl
+++ b/Engine/Plugins/Experimental/EvalGraph/Source/EvalGraph/Public/EvalGraph/EvalGraphConnectionTypeValues.inl
@@ -10,8 +10,98 @@
 
 // NOTE: new types must be added at the bottom to keep serialization from breaking
 
+//
+EVAL_GRAPH_CONNECTION_TYPE(bool, Bool)
+EVAL_GRAPH_CONNECTION_TYPE(char, Char)
 EVAL_GRAPH_CONNECTION_TYPE(int, Integer)
-EVAL_GRAPH_CONNECTION_TYPE(TSharedPtr<FManagedArrayCollection>, ManagedArrayCollection)
+EVAL_GRAPH_CONNECTION_TYPE(uint8, UInt8)
+EVAL_GRAPH_CONNECTION_TYPE(float, Float)
+EVAL_GRAPH_CONNECTION_TYPE(double, Double)
+//
+EVAL_GRAPH_CONNECTION_TYPE(FManagedArrayCollection*,FManagedArrayCollectionPtr)
+EVAL_GRAPH_CONNECTION_TYPE(TManagedArrayBase<double>*,FManagedDoubleArrayPtr)
+EVAL_GRAPH_CONNECTION_TYPE(TManagedArrayBase<float>*,FManagedFloatArrayPtr)
+EVAL_GRAPH_CONNECTION_TYPE(TManagedArrayBase<int32>*,FManagedInt32ArrayPtr)
+EVAL_GRAPH_CONNECTION_TYPE(TManagedArrayBase<bool>*,FManagedBoolArrayPtr)
+EVAL_GRAPH_CONNECTION_TYPE(TManagedArrayBase<FVector3d>*,FManagedVector3dArrayPtr)
+EVAL_GRAPH_CONNECTION_TYPE(TManagedArrayBase<FVector3f>*,FManagedVector3fArrayPtr)
+EVAL_GRAPH_CONNECTION_TYPE(TManagedArrayBase<FTransform>*,FManagedTransformArrayPtr)
+EVAL_GRAPH_CONNECTION_TYPE(TManagedArrayBase<FIntVector>*,FManagedIntVectorArrayPtr)
+EVAL_GRAPH_CONNECTION_TYPE(TManagedArrayBase<FQuat4f>*,FManagedQuat4fArrayPtr)
+EVAL_GRAPH_CONNECTION_TYPE(TManagedArrayBase<FVector2f>*,FManagedVector2fArrayPtr)
+EVAL_GRAPH_CONNECTION_TYPE(TManagedArrayBase<FLinearColor>*,FManagedLinearColorArrayPtr)
+EVAL_GRAPH_CONNECTION_TYPE(TManagedArrayBase<FString>*,FManagedStringArrayPtr)
+EVAL_GRAPH_CONNECTION_TYPE(TManagedArrayBase<FBox>*,FManagedBoxArrayPtr)
+EVAL_GRAPH_CONNECTION_TYPE(TManagedArrayBase<TSet<int32>>*,FManagedIntSetArrayPtr)
+EVAL_GRAPH_CONNECTION_TYPE(TManagedArrayBase<FIntVector4>*,FManagedIntVector4ArrayPtr)
+EVAL_GRAPH_CONNECTION_TYPE(TManagedArrayBase<TArray<FVector2f>>*,FManagedVector2fAArrayPtr)
+EVAL_GRAPH_CONNECTION_TYPE(TManagedArrayBase<TArray<FVector3f>*>*,FManagedVector3fAPtrArrayPtr)
+//
+EVAL_GRAPH_CONNECTION_TYPE(FManagedArrayCollectionSharedSafePtr, ManagedArrayCollectionSharedSafePtr)
+EVAL_GRAPH_CONNECTION_TYPE(FManagedDoubleArraySharedSafePtr, ManagedDoubleArraySharedSafePtr)
+EVAL_GRAPH_CONNECTION_TYPE(FManagedFloatArraySharedSafePtr, ManagedFloatArraySharedSafePtr)
+EVAL_GRAPH_CONNECTION_TYPE(FManagedInt32ArraySharedSafePtr, ManagedInt32ArraySharedSafePtr)
+EVAL_GRAPH_CONNECTION_TYPE(FManagedBoolArraySharedSafePtr, ManagedBoolArraySharedSafePtr)
+EVAL_GRAPH_CONNECTION_TYPE(FManagedVector3dArraySharedSafePtr, ManagedVector3dArraySharedSafePtr)
+EVAL_GRAPH_CONNECTION_TYPE(FManagedVector3fArraySharedSafePtr, ManagedVector3fArraySharedSafePtr)
+EVAL_GRAPH_CONNECTION_TYPE(FManagedTransformArraySharedSafePtr, ManagedTransformArraySharedSafePtr)
+EVAL_GRAPH_CONNECTION_TYPE(FManagedIntVectorArraySharedSafePtr, ManagedIntVectorArraySharedSafePtr)
+EVAL_GRAPH_CONNECTION_TYPE(FManagedQuat4fArraySharedSafePtr, ManagedQuat4fArraySharedSafePtr)
+EVAL_GRAPH_CONNECTION_TYPE(FManagedVector2fArraySharedSafePtr, ManagedVector2fArraySharedSafePtr)
+EVAL_GRAPH_CONNECTION_TYPE(FManagedLinearColorArraySharedSafePtr, ManagedLinearColorArraySharedSafePtr)
+EVAL_GRAPH_CONNECTION_TYPE(FManagedStringArraySharedSafePtr, ManagedStringArraySharedSafePtr)
+EVAL_GRAPH_CONNECTION_TYPE(FManagedBoxArraySharedSafePtr, ManagedBoxArraySharedSafePtr)
+EVAL_GRAPH_CONNECTION_TYPE(FManagedIntSetArraySharedSafePtr, ManagedIntSetArraySharedSafePtr)
+EVAL_GRAPH_CONNECTION_TYPE(FManagedIntVector4ArraySharedSafePtr, ManagedIntVector4ArraySharedSafePtr)
+EVAL_GRAPH_CONNECTION_TYPE(FManagedVector2fAArraySharedSafePtr, ManagedVector2fAArraySharedSafePtr)
+EVAL_GRAPH_CONNECTION_TYPE(FManagedVector3fAPtrArraySharedSafePtr, ManagedVector3fAPtrArraySharedSafePtr)
+//
+EVAL_GRAPH_CONNECTION_TYPE(FString, String)
+EVAL_GRAPH_CONNECTION_TYPE(FBox, Box)
+EVAL_GRAPH_CONNECTION_TYPE(FColor, Color)
+EVAL_GRAPH_CONNECTION_TYPE(FIntVector, IntVector)
+EVAL_GRAPH_CONNECTION_TYPE(FIntVector4, IntVector4)
+EVAL_GRAPH_CONNECTION_TYPE(FLinearColor, LinearColor)
+EVAL_GRAPH_CONNECTION_TYPE(FMatrix, Matrix)
+EVAL_GRAPH_CONNECTION_TYPE(FOrientedBox, OrientedBox)
+EVAL_GRAPH_CONNECTION_TYPE(FPlane, Plane)
+EVAL_GRAPH_CONNECTION_TYPE(FQuat, Quat)
+EVAL_GRAPH_CONNECTION_TYPE(FQuat4f, Quat4f)
+EVAL_GRAPH_CONNECTION_TYPE(FSphere, Sphere)
+EVAL_GRAPH_CONNECTION_TYPE(FTransform, Transform)
+EVAL_GRAPH_CONNECTION_TYPE(FVector, Vector)
+EVAL_GRAPH_CONNECTION_TYPE(FVector2D, Vector2D)
+EVAL_GRAPH_CONNECTION_TYPE(FVector3f, Vector3f)
+EVAL_GRAPH_CONNECTION_TYPE(FVector4, Vector4)
+//
+EVAL_GRAPH_CONNECTION_TYPE(TSet<int32>, IntSet)
+EVAL_GRAPH_CONNECTION_TYPE(TArray<bool>, BoolArray)
+EVAL_GRAPH_CONNECTION_TYPE(TArray<int>, IntArray)
+EVAL_GRAPH_CONNECTION_TYPE(TArray<double>, DoubleArray)
+EVAL_GRAPH_CONNECTION_TYPE(TArray<FBox>, BoxArray)
+EVAL_GRAPH_CONNECTION_TYPE(TArray<FMatrix>, MatrixArray)
+EVAL_GRAPH_CONNECTION_TYPE(TArray<FVector>, VectorArray)
+EVAL_GRAPH_CONNECTION_TYPE(TArray<FVector2f>, Vector2fArray)
+EVAL_GRAPH_CONNECTION_TYPE(TArray<FVector3f>, Vector3fArray)
+EVAL_GRAPH_CONNECTION_TYPE(TArray<FQuat>, QuatArray)
+EVAL_GRAPH_CONNECTION_TYPE(TArray<FQuat4f>, Quat4fArray)
+//
+EVAL_GRAPH_CONNECTION_TYPE(TSet<int32>*, IntSetPtr)
+EVAL_GRAPH_CONNECTION_TYPE(TArray<bool>*, BoolArrayPtr)
+EVAL_GRAPH_CONNECTION_TYPE(TArray<int>*, IntArrayPtr)
+EVAL_GRAPH_CONNECTION_TYPE(TArray<double>*, DoubleArrayPtr)
+EVAL_GRAPH_CONNECTION_TYPE(TArray<FBox>*, BoxArrayPtr)
+EVAL_GRAPH_CONNECTION_TYPE(TArray<FMatrix>*, MatrixArrayPtr)
+EVAL_GRAPH_CONNECTION_TYPE(TArray<FVector>*, VectorArrayPtr)
+EVAL_GRAPH_CONNECTION_TYPE(TArray<FVector2f>*, Vector2fArrayPtr)
+EVAL_GRAPH_CONNECTION_TYPE(TArray<FVector3f>*, Vector3fArrayPtr)
+EVAL_GRAPH_CONNECTION_TYPE(TArray<FQuat>*, QuatArrayPtr)
+EVAL_GRAPH_CONNECTION_TYPE(TArray<FQuat4f>*, Quat4fArrayPtr)
+//
+EVAL_GRAPH_CONNECTION_TYPE(Chaos::FImplicitObject3*, ImplicitObject3Ptr)
+EVAL_GRAPH_CONNECTION_TYPE(Chaos::FBVHParticlesFloat3*, BVHParticlesFloat3Ptr)
+EVAL_GRAPH_CONNECTION_TYPE(FImplicitObject3SerializablePtr, FImplicitObject3SerializablePtr)
+EVAL_GRAPH_CONNECTION_TYPE(FImplicitObjectSharedSafePtr, FImplicitObjectSharedSafePtr)
 
 // NOTE: new types must be added at the bottom to keep serialization from breaking
 
diff --git a/Engine/Plugins/Experimental/EvalGraph/Source/EvalGraph/Public/EvalGraph/EvalGraphConnectionTypes.h b/Engine/Plugins/Experimental/EvalGraph/Source/EvalGraph/Public/EvalGraph/EvalGraphConnectionTypes.h
--- a/Engine/Plugins/Experimental/EvalGraph/Source/EvalGraph/Public/EvalGraph/EvalGraphConnectionTypes.h
+++ b/Engine/Plugins/Experimental/EvalGraph/Source/EvalGraph/Public/EvalGraph/EvalGraphConnectionTypes.h
@@ -2,8 +2,60 @@
 #pragma once
 #include "GeometryCollection/ManagedArrayCollection.h"
 
+#include "Math/Box.h"
+#include "Math/Box2D.h"
+#include "Math/Color.h"
+#include "Math/IntVector.h"
+#include "Math/Matrix.h"
+#include "Math/TransformCalculus2D.h"
+#include "Math/OrientedBox.h"
+#include "Math/OrthoMatrix.h"
+#include "Math/Plane.h"
+#include "Math/Quat.h"
+#include "Math/TransformCalculus2D.h"
+#include "Math/QuatRotationTranslationMatrix.h"
+#include "Math/Sphere.h"
+#include "Math/TransformNonVectorized.h"
+#include "Math/TransformCalculus2D.h"
+#include "Math/TranslationMatrix.h"
+#include "Math/Vector.h"
+#include "Math/Vector2D.h"
+#include "Math/Vector4.h"
+#include "Math/Matrix.h"
+#include "Math/RotationMatrix.h"
+
+#include "Chaos/ImplicitObject.h"
+#include "Chaos/BVHParticles.h"
+#include "Chaos/ParticleHandle.h"
+#include "Chaos/ParticleHandleFwd.h"
+#include "Chaos/Convex.h"
+
+
 namespace Eg
 {
+	typedef TSharedPtr<FManagedArrayCollection, ESPMode::ThreadSafe> FManagedArrayCollectionSharedSafePtr;
+	typedef TSharedPtr<TManagedArrayBase<double>, ESPMode::ThreadSafe> FManagedDoubleArraySharedSafePtr;
+	typedef TSharedPtr<TManagedArrayBase<float>, ESPMode::ThreadSafe> FManagedFloatArraySharedSafePtr;
+	typedef TSharedPtr<TManagedArrayBase<int32>, ESPMode::ThreadSafe> FManagedInt32ArraySharedSafePtr;
+	typedef TSharedPtr<TManagedArrayBase<bool>, ESPMode::ThreadSafe> FManagedBoolArraySharedSafePtr;
+	typedef TSharedPtr<TManagedArrayBase<FVector3d>, ESPMode::ThreadSafe> FManagedVector3dArraySharedSafePtr;
+	typedef TSharedPtr<TManagedArrayBase<FVector3f>, ESPMode::ThreadSafe> FManagedVector3fArraySharedSafePtr;
+	typedef TSharedPtr<TManagedArrayBase<FTransform>, ESPMode::ThreadSafe> FManagedTransformArraySharedSafePtr;
+	typedef TSharedPtr<TManagedArrayBase<FIntVector>, ESPMode::ThreadSafe> FManagedIntVectorArraySharedSafePtr;
+	typedef TSharedPtr<TManagedArrayBase<FQuat4f>, ESPMode::ThreadSafe> FManagedQuat4fArraySharedSafePtr;
+	typedef TSharedPtr<TManagedArrayBase<FVector2f>, ESPMode::ThreadSafe> FManagedVector2fArraySharedSafePtr;
+	typedef TSharedPtr<TManagedArrayBase<FLinearColor>, ESPMode::ThreadSafe> FManagedLinearColorArraySharedSafePtr;
+	typedef TSharedPtr<TManagedArrayBase<FString>, ESPMode::ThreadSafe> FManagedStringArraySharedSafePtr;
+	typedef TSharedPtr<TManagedArrayBase<FBox>, ESPMode::ThreadSafe> FManagedBoxArraySharedSafePtr;
+	typedef TSharedPtr<TManagedArrayBase<TSet<int32>>, ESPMode::ThreadSafe> FManagedIntSetArraySharedSafePtr;
+	typedef TSharedPtr<TManagedArrayBase<FIntVector4>, ESPMode::ThreadSafe> FManagedIntVector4ArraySharedSafePtr;
+	typedef TSharedPtr<TManagedArrayBase<TArray<FVector2f>>, ESPMode::ThreadSafe> FManagedVector2fAArraySharedSafePtr;
+	typedef TSharedPtr<TManagedArrayBase<TArray<FVector3f>*>, ESPMode::ThreadSafe> FManagedVector3fAPtrArraySharedSafePtr;
+
+
+	typedef Chaos::TSerializablePtr<Chaos::FImplicitObject3> FImplicitObject3SerializablePtr;
+	typedef TSharedPtr<Chaos::FImplicitObject, ESPMode::ThreadSafe> FImplicitObjectSharedSafePtr;
+
 	// ---------------------------------------------------------
 	//
 	// General purpose EManagedArrayType definition. 
