Replacing legacy EditorStyle call with AppStyle call

diff --git a/Engine/Plugins/Experimental/EvalGraph/Source/EvalGraphEditor/Public/EvalGraph/EvalGraphEditorActions.h b/Engine/Plugins/Experimental/EvalGraph/Source/EvalGraphEditor/Public/EvalGraph/EvalGraphEditorActions.h
--- a/Engine/Plugins/Experimental/EvalGraph/Source/EvalGraphEditor/Public/EvalGraph/EvalGraphEditorActions.h
+++ b/Engine/Plugins/Experimental/EvalGraph/Source/EvalGraphEditor/Public/EvalGraph/EvalGraphEditorActions.h
@@ -4,14 +4,14 @@
 
 #include "CoreMinimal.h"
 #include "Framework/Commands/Commands.h"
-#include "EditorStyleSet.h"
+#include "Styling/AppStyle.h"
 
 class FEvalGraphEditorCommandsImpl : public TCommands<FEvalGraphEditorCommandsImpl>
 {
 public:
 
 	FEvalGraphEditorCommandsImpl()
-		: TCommands<FEvalGraphEditorCommandsImpl>( TEXT("EvalGraphEditor"), NSLOCTEXT("Contexts", "EvalGraphEditor", "Scene Graph Editor"), NAME_None, FEditorStyle::GetStyleSetName() )
+		: TCommands<FEvalGraphEditorCommandsImpl>( TEXT("EvalGraphEditor"), NSLOCTEXT("Contexts", "EvalGraphEditor", "Scene Graph Editor"), NAME_None, FAppStyle::GetAppStyleSetName() )
 	{
 	}	
 
