EvalGraph -  Stub graph editor

diff --git a/Engine/Plugins/Experimental/EvalGraph/Source/EvalGraphEditor/Private/EvalGraph/EvalGraphAssetActions.cpp b/Engine/Plugins/Experimental/EvalGraph/Source/EvalGraphEditor/Private/EvalGraph/EvalGraphAssetActions.cpp
--- a/Engine/Plugins/Experimental/EvalGraph/Source/EvalGraphEditor/Private/EvalGraph/EvalGraphAssetActions.cpp
+++ b/Engine/Plugins/Experimental/EvalGraph/Source/EvalGraphEditor/Private/EvalGraph/EvalGraphAssetActions.cpp
@@ -1,6 +1,7 @@
 // Copyright Epic Games, Inc. All Rights Reserved.
 
 #include "EvalGraph/EvalGraphAssetActions.h"
+#include "EvalGraph/EvalGraphEditorPlugin.h"
 
 #include "EvalGraph/EvalGraphObject.h"
 
@@ -27,11 +28,20 @@ void FEvalGraphAssetActions::GetActions(const TArray<UObject*>& InObjects,
 }
 
 void FEvalGraphAssetActions::OpenAssetEditor(
-	const TArray<UObject*>& InObjects, TSharedPtr<IToolkitHost> EditWithinLevelEditor /*= TSharedPtr<IToolkitHost>()*/)
+	const TArray<UObject*>& InObjects, TSharedPtr<IToolkitHost> EditWithinLevelEditor)
 {
-	FSimpleAssetEditor::CreateEditor(EToolkitMode::Standalone, EditWithinLevelEditor, InObjects);
+	EToolkitMode::Type Mode = EditWithinLevelEditor.IsValid() ? EToolkitMode::WorldCentric : EToolkitMode::Standalone;
+	for (auto ObjIt = InObjects.CreateConstIterator(); ObjIt; ++ObjIt)
+	{
+		if (auto Object = Cast<UEvalGraph>(*ObjIt))
+		{
+			IEvalGraphEditorPlugin* EvalGraphEditorPlugin = &FModuleManager::LoadModuleChecked<IEvalGraphEditorPlugin>("EvalGraphEditor");
+			EvalGraphEditorPlugin->CreateEvalGraphAssetEditor(Mode, EditWithinLevelEditor, Object);
+		}
+	}
 }
 
+
 uint32 FEvalGraphAssetActions::GetCategories()
 {
 	return EAssetTypeCategories::Physics;
diff --git a/Engine/Plugins/Experimental/EvalGraph/Source/EvalGraphEditor/Private/EvalGraph/EvalGraphAssetFactory.cpp b/Engine/Plugins/Experimental/EvalGraph/Source/EvalGraphEditor/Private/EvalGraph/EvalGraphAssetFactory.cpp
--- a/Engine/Plugins/Experimental/EvalGraph/Source/EvalGraphEditor/Private/EvalGraph/EvalGraphAssetFactory.cpp
+++ b/Engine/Plugins/Experimental/EvalGraph/Source/EvalGraphEditor/Private/EvalGraph/EvalGraphAssetFactory.cpp
@@ -3,6 +3,11 @@
 #include "EvalGraph/EvalGraphAssetFactory.h"
 
 #include "EvalGraph/EvalGraphObject.h"
+#include "HAL/IConsoleManager.h"
+
+bool bEvalGraphEnableCreation = false;
+FAutoConsoleVariableRef CVarEvalGraphEnableCreation(TEXT("c.EvalGraphEnableCreation"), bEvalGraphEnableCreation, TEXT("Enable creation for evaluation graph (Dev-Only)"));
+
 
 UEvalGraphAssetFactory::UEvalGraphAssetFactory()
 {
@@ -11,7 +16,7 @@ UEvalGraphAssetFactory::UEvalGraphAssetFactory()
 
 bool UEvalGraphAssetFactory::CanCreateNew() const
 {
-	return true;
+	return bEvalGraphEnableCreation;
 }
 
 bool UEvalGraphAssetFactory::FactoryCanImport(const FString& Filename)
@@ -27,7 +32,7 @@ UObject* UEvalGraphAssetFactory::FactoryCreateNew(UClass* InClass, UObject* InPa
 
 bool UEvalGraphAssetFactory::ShouldShowInNewMenu() const
 {
-	return true;
+	return bEvalGraphEnableCreation;
 }
 
 bool UEvalGraphAssetFactory::ConfigureProperties()
diff --git a/Engine/Plugins/Experimental/EvalGraph/Source/EvalGraphEditor/Private/EvalGraph/EvalGraphEditorActions.cpp b/Engine/Plugins/Experimental/EvalGraph/Source/EvalGraphEditor/Private/EvalGraph/EvalGraphEditorActions.cpp
--- a/Engine/Plugins/Experimental/EvalGraph/Source/EvalGraphEditor/Private/EvalGraph/EvalGraphEditorActions.cpp
+++ b/Engine/Plugins/Experimental/EvalGraph/Source/EvalGraphEditor/Private/EvalGraph/EvalGraphEditorActions.cpp
@@ -1,12 +1,33 @@
 // Copyright Epic Games, Inc. All Rights Reserved.
 
 #include "EvalGraph/EvalGraphEditorActions.h"
+#include "EvalGraph/EvalGraphNodeFactory.h"
 
 #define LOCTEXT_NAMESPACE "EvalGraphEditorCommands"
 
 void FEvalGraphEditorCommandsImpl::RegisterCommands()
 {
-	UI_COMMAND(EvaluateNode, "Evaluate", "Trigger an evaluation of the selected node.", EUserInterfaceActionType::Button, FInputChord())
+	UI_COMMAND(EvaluateNode, "Evaluate", "Trigger an evaluation of the selected node.", EUserInterfaceActionType::Button, FInputChord());
+
+	if (Eg::FNodeFactory* Factory = Eg::FNodeFactory::GetInstance())
+	{
+		for (FName NodeName : Factory->RegisteredNodes())
+		{
+			TSharedPtr< FUICommandInfo > AddNode;
+			FUICommandInfo::MakeCommandInfo(
+				this->AsShared(),
+				AddNode,
+				NodeName, //FName("UseCreationFormToggle"),
+				NSLOCTEXT("DataFlow", "DataflowButton", "New Dataflow Node"),
+				NSLOCTEXT("DataFlow", "NewDataflowNodeTooltip", "New Dataflow Node Tooltip"),
+				FSlateIcon(),
+				EUserInterfaceActionType::Button,
+				FInputChord()
+			);
+			//UI_COMMAND(AddNode,"Name Here", "Add Node", EUserInterfaceActionType::Button, FInputChord());
+			CreateNodesMap.Add(NodeName, AddNode);
+		}
+	}
 }
 
 void FEvalGraphEditorCommands::Register()
diff --git a/Engine/Plugins/Experimental/EvalGraph/Source/EvalGraphEditor/Private/EvalGraph/EvalGraphEditorPlugin.cpp b/Engine/Plugins/Experimental/EvalGraph/Source/EvalGraphEditor/Private/EvalGraph/EvalGraphEditorPlugin.cpp
--- a/Engine/Plugins/Experimental/EvalGraph/Source/EvalGraphEditor/Private/EvalGraph/EvalGraphEditorPlugin.cpp
+++ b/Engine/Plugins/Experimental/EvalGraph/Source/EvalGraphEditor/Private/EvalGraph/EvalGraphEditorPlugin.cpp
@@ -5,8 +5,10 @@
 #include "AssetToolsModule.h"
 #include "CoreMinimal.h"
 #include "EdGraphUtilities.h"
+#include "EvalGraphEditorToolkit.h"
+#include "EvalGraph/EvalGraphNodeFactory.h"
 #include "EvalGraph/EvalGraphAssetActions.h"
-#include "EvalGraph/EvalGraphNodeFactories.h"
+#include "EvalGraph/EvalGraphSNodeFactories.h"
 
 #define LOCTEXT_NAMESPACE "EvalGraphEditor"
 
@@ -22,8 +24,8 @@ void IEvalGraphEditorPlugin::StartupModule()
 	AssetTools.RegisterAssetTypeActions(MakeShareable(EvalGraphAssetActions));
 
 
-	//EvalGraphNodeFactory = MakeShareable(new FEvalGraphNodeFactory());
-	//FEdGraphUtilities::RegisterVisualNodeFactory(EvalGraphNodeFactory);
+	EvalGraphSNodeFactory = MakeShareable(new FEvalGraphSNodeFactory());
+	FEdGraphUtilities::RegisterVisualNodeFactory(EvalGraphSNodeFactory);
 }
 
 void IEvalGraphEditorPlugin::ShutdownModule()
@@ -35,10 +37,18 @@ void IEvalGraphEditorPlugin::ShutdownModule()
 
 		AssetTools.UnregisterAssetTypeActions(EvalGraphAssetActions->AsShared());
 
-		//FEdGraphUtilities::UnregisterVisualNodeFactory(EvalGraphNodeFactory);
+		FEdGraphUtilities::UnregisterVisualNodeFactory(EvalGraphSNodeFactory);
 	}
 }
 
+TSharedRef<FAssetEditorToolkit> IEvalGraphEditorPlugin::CreateEvalGraphAssetEditor(const EToolkitMode::Type Mode, const TSharedPtr<IToolkitHost>& InitToolkitHost, UObject* FleshAsset)
+{
+	TSharedPtr<FEvalGraphEditorToolkit> NewEvalGraphAssetEditor = MakeShared<FEvalGraphEditorToolkit>();
+	NewEvalGraphAssetEditor->InitEvalGraphEditor(Mode, InitToolkitHost, FleshAsset);
+	return StaticCastSharedPtr<FAssetEditorToolkit>(NewEvalGraphAssetEditor).ToSharedRef();
+}
+
+
 IMPLEMENT_MODULE(IEvalGraphEditorPlugin, EvalGraphEditor)
 
 
diff --git a/Engine/Plugins/Experimental/EvalGraph/Source/EvalGraphEditor/Private/EvalGraph/EvalGraphEditorToolkit.cpp b/Engine/Plugins/Experimental/EvalGraph/Source/EvalGraphEditor/Private/EvalGraph/EvalGraphEditorToolkit.cpp
new file mode 100644
--- /dev/null
+++ b/Engine/Plugins/Experimental/EvalGraph/Source/EvalGraphEditor/Private/EvalGraph/EvalGraphEditorToolkit.cpp
@@ -0,0 +1,259 @@
+// Copyright Epic Games, Inc. All Rights Reserved.
+
+#include "EvalGraphEditorToolkit.h"
+
+#include "EvalGraph/EvalGraphEditorActions.h"
+#include "EvalGraph/EvalGraphEdNode.h"
+#include "EvalGraph/EvalGraphNodeFactory.h"
+#include "EvalGraph/EvalGraphObject.h"
+#include "EvalGraph/EvalGraphSchema.h"
+#include "EvalGraph/EvalGraph.h"
+#include "EditorStyleSet.h"
+#include "EditorViewportTabContent.h"
+#include "EditorViewportLayout.h"
+#include "EditorViewportCommands.h"
+#include "GraphEditorActions.h"
+#include "Modules/ModuleManager.h"
+#include "PropertyEditorModule.h"
+#include "Styling/SlateStyleRegistry.h"
+
+#define LOCTEXT_NAMESPACE "EvalGraphEditorToolkit"
+
+//DEFINE_LOG_CATEGORY_STATIC(FEvalGraphEditorToolkitLog, Log, All);
+
+
+const FName FEvalGraphEditorToolkit::GraphCanvasTabId(TEXT("EvalGraphEditor_GraphCanvas"));
+const FName FEvalGraphEditorToolkit::PropertiesTabId(TEXT("EvalGraphEditor_Properties"));
+
+void FEvalGraphEditorToolkit::InitEvalGraphEditor(const EToolkitMode::Type Mode, const TSharedPtr<IToolkitHost>& InitToolkitHost, UObject* ObjectToEdit)
+{
+	EvalGraph = CastChecked<UEvalGraph>(ObjectToEdit);
+	if (EvalGraph != nullptr)
+	{
+		EvalGraph->Schema = UEvalGraphSchema::StaticClass();
+
+		GraphEditor = CreateGraphEditorWidget(EvalGraph);
+		PropertiesEditor = CreatePropertiesEditorWidget(ObjectToEdit);
+
+		const TSharedRef<FTabManager::FLayout> StandaloneDefaultLayout = FTabManager::NewLayout("EvalGraph_Layout")
+			->AddArea
+			(
+				FTabManager::NewPrimaryArea()->SetOrientation(Orient_Vertical)
+				->Split
+				(
+					FTabManager::NewSplitter()->SetOrientation(Orient_Vertical)
+					->Split
+					(
+						FTabManager::NewSplitter()->SetOrientation(Orient_Horizontal)
+						->SetSizeCoefficient(0.9f)
+						->Split
+						(
+							FTabManager::NewStack()
+							->SetSizeCoefficient(0.6f)
+							->AddTab(GraphCanvasTabId, ETabState::OpenedTab)
+						)
+						->Split
+						(
+							FTabManager::NewSplitter()->SetOrientation(Orient_Vertical)
+							->SetSizeCoefficient(0.2f)
+							->Split
+							(
+								FTabManager::NewStack()
+								->SetSizeCoefficient(0.7f)
+								->AddTab(PropertiesTabId, ETabState::OpenedTab)
+							)
+						)
+					)
+				)
+			);
+
+		const bool bCreateDefaultStandaloneMenu = true;
+		const bool bCreateDefaultToolbar = true;
+		FAssetEditorToolkit::InitAssetEditor(Mode, InitToolkitHost, FName(TEXT("EvalGraphEditorApp")), StandaloneDefaultLayout, bCreateDefaultToolbar, bCreateDefaultStandaloneMenu, ObjectToEdit);
+	}
+}
+
+void FEvalGraphEditorToolkit::EvaluateNode()
+{
+	if (UEvalGraph* Graph = dynamic_cast<UEvalGraph*>(GraphEditor->GetCurrentGraph()))
+	{
+		for (UObject* Ode : GetSelectedNodes())
+		{
+			if (UEvalGraphEdNode* Node = dynamic_cast<UEvalGraphEdNode*>(Ode))
+			{
+				//UE_LOG(FEvalGraphEditorToolkitLog, Verbose, TEXT("FEvalGraphEditorToolkit::EvaluateNode(%s)"), *Ode->GetFullName());
+				//Graph->Evaluate(Node);
+			}
+		}
+	}
+}
+void FEvalGraphEditorToolkit::CreateNode(FName RegisteredNode)
+{
+	//UE_LOG(FEvalGraphEditorToolkitLog, Verbose, TEXT("FEvalGraphEditorToolkit::CreateNode"), RegisteredNode);
+	if (UEvalGraph* Graph = dynamic_cast<UEvalGraph*>(GraphEditor->GetCurrentGraph()))
+	{
+		if (UEvalGraph * Asset = GetEvalGraph())
+		{
+			//TSharedPtr<FNewManagedArrayCollectionNode> Node1 = AddNode(new FNewManagedArrayCollectionNode({ FName("Node1") }));
+		}
+	}
+}
+
+
+FGraphPanelSelectionSet FEvalGraphEditorToolkit::GetSelectedNodes() const
+{
+	if (GraphEditor.IsValid())
+	{
+		return GraphEditor->GetSelectedNodes();
+	}
+	return FGraphPanelSelectionSet();
+}
+
+void FEvalGraphEditorToolkit::OnSelectedNodesChanged(const TSet<UObject*>& NewSelection)
+{
+}
+
+
+
+TSharedRef<SGraphEditor> FEvalGraphEditorToolkit::CreateGraphEditorWidget(UEvalGraph* EvalGraphToEdit)
+{
+	ensure(EvalGraphToEdit);
+
+	FEvalGraphEditorCommands::Register();
+	FGraphEditorCommands::Register();
+
+	// No need to regenerate the commands.
+	if (!GraphEditorCommands.IsValid())
+	{
+		GraphEditorCommands = MakeShareable(new FUICommandList);
+		{
+			//if (Eg::FNodeFactory* Factory = Eg::FNodeFactory::GetInstance())
+			//{
+			//	for (FName NodeName : Factory->RegisteredNodes())
+			//	{
+			//		if (FEvalGraphEditorCommands::Get().CreateNodesMap.Contains(NodeName))
+			//		{
+			//			GraphEditorCommands->MapAction(FEvalGraphEditorCommands::Get().CreateNodesMap[NodeName],
+			//				FExecuteAction::CreateSP(this, &FEvalGraphEditorToolkit::CreateNode, NodeName));
+			//		}
+//
+	//			}
+	//		}
+
+			GraphEditorCommands->MapAction(FEvalGraphEditorCommands::Get().EvaluateNode,
+				FExecuteAction::CreateSP(this, &FEvalGraphEditorToolkit::EvaluateNode)
+			);
+		}
+	}
+
+
+	FGraphAppearanceInfo AppearanceInfo;
+	AppearanceInfo.CornerText = LOCTEXT("AppearanceCornerText_EvalGraphEditor", "EvalGraph Graph");
+
+	SGraphEditor::FGraphEditorEvents InEvents;
+	InEvents.OnSelectionChanged = SGraphEditor::FOnSelectionChanged::CreateSP(this, &FEvalGraphEditorToolkit::OnSelectedNodesChanged);
+
+	return SNew(SGraphEditor)
+		.AdditionalCommands(GraphEditorCommands)
+		.IsEditable(true)
+		.Appearance(AppearanceInfo)
+		.GraphToEdit(EvalGraphToEdit)
+		.GraphEvents(InEvents)
+		.ShowGraphStateOverlay(false);
+}
+
+TSharedPtr<IDetailsView> FEvalGraphEditorToolkit::CreatePropertiesEditorWidget(UObject* ObjectToEdit)
+{
+	ensure(ObjectToEdit);
+	FPropertyEditorModule& PropertyEditorModule = FModuleManager::LoadModuleChecked<FPropertyEditorModule>(TEXT("PropertyEditor"));
+
+	FDetailsViewArgs DetailsViewArgs;
+	DetailsViewArgs.bAllowSearch = true;
+	DetailsViewArgs.bLockable = false;
+	DetailsViewArgs.bUpdatesFromSelection = false;
+	DetailsViewArgs.NameAreaSettings = FDetailsViewArgs::HideNameArea;
+	DetailsViewArgs.NotifyHook = this;
+
+	TSharedPtr<IDetailsView> DetailsView = PropertyEditorModule.CreateDetailView(DetailsViewArgs);
+	DetailsView->SetObject(ObjectToEdit);
+	return DetailsView;
+
+}
+
+
+TSharedRef<SDockTab> FEvalGraphEditorToolkit::SpawnTab_GraphCanvas(const FSpawnTabArgs& Args)
+{
+	check(Args.GetTabId() == GraphCanvasTabId);
+
+	TSharedRef<SDockTab> SpawnedTab = SNew(SDockTab)
+		.Label(LOCTEXT("EvalGraphEditor_EvalGraph_TabTitle", "Graph"))
+		[
+			GraphEditor.ToSharedRef()
+		];
+
+	return SpawnedTab;
+}
+
+TSharedRef<SDockTab> FEvalGraphEditorToolkit::SpawnTab_Properties(const FSpawnTabArgs& Args)
+{
+	check(Args.GetTabId() == PropertiesTabId);
+
+	return SNew(SDockTab)
+		.Label(LOCTEXT("EvalGraphEditor_Properties_TabTitle", "Details"))
+		[
+			PropertiesEditor.ToSharedRef()
+		];
+}
+
+
+void FEvalGraphEditorToolkit::RegisterTabSpawners(const TSharedRef<FTabManager>& InTabManager)
+{
+	TSharedRef<FWorkspaceItem> WorkspaceMenuCategoryRef = InTabManager->AddLocalWorkspaceMenuCategory(LOCTEXT("WorkspaceMenu_EvalGraphEditor", "EvalGraph Editor"));
+
+	InTabManager->RegisterTabSpawner(GraphCanvasTabId, FOnSpawnTab::CreateSP(this, &FEvalGraphEditorToolkit::SpawnTab_GraphCanvas))
+		.SetDisplayName(LOCTEXT("EvalGraphTab", "Graph"))
+		.SetGroup(WorkspaceMenuCategoryRef)
+		.SetIcon(FSlateIcon(FAppStyle::GetAppStyleSetName(), "GraphEditor.EventGraph_16x"));
+
+	InTabManager->RegisterTabSpawner(PropertiesTabId, FOnSpawnTab::CreateSP(this, &FEvalGraphEditorToolkit::SpawnTab_Properties))
+		.SetDisplayName(LOCTEXT("PropertiesTab", "Details"))
+		.SetGroup(WorkspaceMenuCategoryRef)
+		.SetIcon(FSlateIcon(FAppStyle::GetAppStyleSetName(), "LevelEditor.Tabs.Details"));
+
+	FAssetEditorToolkit::RegisterTabSpawners(InTabManager);
+}
+
+
+FName FEvalGraphEditorToolkit::GetToolkitFName() const
+{
+	return FName("EvalGraphEditor");
+}
+
+FText FEvalGraphEditorToolkit::GetBaseToolkitName() const
+{
+	return LOCTEXT("AppLabel", "EvalGraph Editor");
+}
+
+FString FEvalGraphEditorToolkit::GetWorldCentricTabPrefix() const
+{
+	return LOCTEXT("WorldCentricTabPrefix", "EvalGraph").ToString();
+}
+
+FLinearColor FEvalGraphEditorToolkit::GetWorldCentricTabColorScale() const
+{
+	return FLinearColor(0.3f, 0.2f, 0.5f, 0.5f);
+}
+
+FString FEvalGraphEditorToolkit::GetReferencerName() const
+{
+	return TEXT("EvalGraphEditorToolkit");
+}
+
+void FEvalGraphEditorToolkit::AddReferencedObjects(FReferenceCollector& Collector)
+{
+	if (EvalGraph)
+	{
+		Collector.AddReferencedObject(EvalGraph);
+	}
+}
+#undef LOCTEXT_NAMESPACE
diff --git a/Engine/Plugins/Experimental/EvalGraph/Source/EvalGraphEditor/Private/EvalGraph/EvalGraphEditorToolkit.h b/Engine/Plugins/Experimental/EvalGraph/Source/EvalGraphEditor/Private/EvalGraph/EvalGraphEditorToolkit.h
new file mode 100644
--- /dev/null
+++ b/Engine/Plugins/Experimental/EvalGraph/Source/EvalGraphEditor/Private/EvalGraph/EvalGraphEditorToolkit.h
@@ -0,0 +1,60 @@
+// Copyright Epic Games, Inc. All Rights Reserved.
+
+#pragma once
+
+#include "CoreMinimal.h"
+#include "Toolkits/SimpleAssetEditor.h"
+#include "Misc/NotifyHook.h"
+
+class FEditorViewportTabContent;
+class IDetailsView;
+class FTabManager;
+class IToolkitHost;
+class UEvalGraph;
+class SEvalGraphEditorViewport;
+
+class FEvalGraphEditorToolkit : public FAssetEditorToolkit, public FNotifyHook, public FGCObject
+{
+public:
+	void InitEvalGraphEditor(const EToolkitMode::Type Mode, const TSharedPtr<IToolkitHost>& InitToolkitHost, UObject* ObjectToEdit);
+
+	// IToolkit Interface
+	virtual FName GetToolkitFName() const override;
+	virtual FText GetBaseToolkitName() const override;
+	virtual FString GetWorldCentricTabPrefix() const override;
+	virtual FLinearColor GetWorldCentricTabColorScale() const override;
+
+	// FGCObject Interface
+	virtual void AddReferencedObjects(FReferenceCollector& Collector) override;
+	virtual FString GetReferencerName() const override; 
+
+	// Tab spawners 
+	virtual void RegisterTabSpawners(const TSharedRef<FTabManager>& TabManager) override;
+	TSharedRef<SDockTab> SpawnTab_GraphCanvas(const FSpawnTabArgs& Args);
+	TSharedRef<SDockTab> SpawnTab_Properties(const FSpawnTabArgs& Args);
+
+	// Graph Editor Operations
+	void OnSelectedNodesChanged(const TSet<UObject*>& NewSelection);
+	void EvaluateNode();
+	void CreateNode(FName RegisteredNode);
+	FGraphPanelSelectionSet GetSelectedNodes() const;
+
+	UEvalGraph* GetEvalGraph() { return EvalGraph; }
+	const UEvalGraph* GetEvalGraph() const { return EvalGraph; }
+
+private:
+
+	static const FName GraphCanvasTabId;
+	TSharedPtr<SGraphEditor> GraphEditor;
+	TSharedRef<SGraphEditor> CreateGraphEditorWidget(UEvalGraph* ObjectToEdit);
+
+	static const FName PropertiesTabId;
+	TSharedPtr<IDetailsView> PropertiesEditor;
+	TSharedPtr<IDetailsView> CreatePropertiesEditorWidget(UObject* ObjectToEdit);
+
+
+	UEvalGraph* EvalGraph = nullptr;
+
+	/** Command list for the graph editor */
+	TSharedPtr<FUICommandList> GraphEditorCommands;
+};
diff --git a/Engine/Plugins/Experimental/EvalGraph/Source/EvalGraphEditor/Private/EvalGraph/EvalGraphNodeFactories.cpp b/Engine/Plugins/Experimental/EvalGraph/Source/EvalGraphEditor/Private/EvalGraph/EvalGraphNodeFactories.cpp
deleted file mode 100644
diff --git a/Engine/Plugins/Experimental/EvalGraph/Source/EvalGraphEditor/Private/EvalGraph/EvalGraphSNode.cpp b/Engine/Plugins/Experimental/EvalGraph/Source/EvalGraphEditor/Private/EvalGraph/EvalGraphSNode.cpp
new file mode 100644
--- /dev/null
+++ b/Engine/Plugins/Experimental/EvalGraph/Source/EvalGraphEditor/Private/EvalGraph/EvalGraphSNode.cpp
@@ -0,0 +1,97 @@
+// Copyright Epic Games, Inc. All Rights Reserved.
+
+#include "EvalGraph/EvalGraphSNode.h"
+
+#include "EvalGraph/EvalGraphEdNode.h"
+#include "EvalGraph/EvalGraphNodeFactory.h"
+#include "EvalGraph/EvalGraphObject.h"
+#include "EvalGraph/EvalGraph.h"
+#include "Logging/LogMacros.h"
+#include "Widgets/Layout/SBorder.h"
+
+#define LOCTEXT_NAMESPACE "SEvalGraphEdNode"
+//
+// SEvalGraphEdNode
+//
+
+void SEvalGraphEdNode::Construct(const FArguments& InArgs, UEvalGraphEdNode* InNode)
+{
+	GraphNode = InNode;
+	UEvalGraph* Graph = Cast<UEvalGraph>(InNode->GetOuter());
+	if( ensure(Graph) )
+	{
+		if (Eg::FNodeFactory* Factory = Eg::FNodeFactory::GetInstance())
+		{
+			if (TSharedPtr<Eg::FNode> EgNode = Factory->NewNodeFromRegisteredType(*Graph->GetEvalGraph(), { FGuid::NewGuid(),FName(InNode->GetName()),FName(InNode->GetName()) }))
+			{
+				InNode->SetEgGraph(Graph->GetEvalGraph());
+				InNode->SetEgNode(EgNode->GetGuid());
+				InNode->AllocateDefaultPins();
+			}
+		}
+	}
+	UpdateGraphNode();
+}
+
+FReply SEvalGraphEdNode::OnMouseButtonDoubleClick(const FGeometry& InMyGeometry, const FPointerEvent& InMouseEvent)
+{
+	return Super::OnMouseButtonDoubleClick(InMyGeometry, InMouseEvent);
+}
+
+
+
+//
+// FAssetSchemaAction_EvalGraph_CreateNode_Example
+//
+TSharedPtr<FAssetSchemaAction_EvalGraph_CreateNode_EvalGraphEdNode> FAssetSchemaAction_EvalGraph_CreateNode_EvalGraphEdNode::CreateAction(UEdGraph* Owner, const FName & InNodeTypeName)
+{
+	const FText AddToolTip = LOCTEXT("EvalGraphNodeTooltip_Example", "Add a Dataflow node.");
+	const FText NodeName = FText::FromString(*InNodeTypeName.ToString());
+	const FText Catagory = LOCTEXT("EvalGraphNodeDescription_Example", "Dataflow"); 	
+	TSharedPtr<FAssetSchemaAction_EvalGraph_CreateNode_EvalGraphEdNode> NewNodeAction(
+		new FAssetSchemaAction_EvalGraph_CreateNode_EvalGraphEdNode(Catagory, NodeName, AddToolTip, 0));
+	NewNodeAction->NodeTemplate = NewObject<UEvalGraphEdNode>(Owner, UEvalGraphEdNode::StaticClass(), InNodeTypeName);
+	//NewNodeAction->NodeTemplate->GenericGraphNode = NewObject<UGenericGraphNode>(NewNodeAction->NodeTemplate, Graph->NodeType);
+	//NewNodeAction->NodeTemplate->GenericGraphNode->Graph = Graph;
+	return NewNodeAction;
+}
+
+
+UEdGraphNode* FAssetSchemaAction_EvalGraph_CreateNode_EvalGraphEdNode::PerformAction(class UEdGraph* ParentGraph, UEdGraphPin* FromPin, const FVector2D Location, bool bSelectNewNode)
+{
+	UEdGraphNode* ResultNode = nullptr;
+
+	if (NodeTemplate != nullptr)
+	{
+		//const FScopedTransaction Transaction(LOCTEXT("EvalGraphNewNode", "Generic Graph Editor: New Node"));
+		ParentGraph->Modify();
+		if (FromPin != nullptr)
+			FromPin->Modify();
+
+		NodeTemplate->Rename(nullptr, ParentGraph);
+		ParentGraph->AddNode(NodeTemplate, true, bSelectNewNode);
+
+		NodeTemplate->CreateNewGuid();
+		NodeTemplate->PostPlacedNewNode();
+		//NodeTemplate->AllocateDefaultPins();
+		NodeTemplate->AutowireNewNode(FromPin);
+
+		NodeTemplate->NodePosX = Location.X;
+		NodeTemplate->NodePosY = Location.Y;
+
+		//NodeTemplate->GenericGraphNode->SetFlags(RF_Transactional);
+		NodeTemplate->SetFlags(RF_Transactional);
+
+		ResultNode = NodeTemplate;
+	}
+
+	return ResultNode;
+}
+
+void FAssetSchemaAction_EvalGraph_CreateNode_EvalGraphEdNode::AddReferencedObjects(FReferenceCollector& Collector)
+{
+	FEdGraphSchemaAction::AddReferencedObjects(Collector);
+	Collector.AddReferencedObject(NodeTemplate);
+}
+
+#undef LOCTEXT_NAMESPACE
diff --git a/Engine/Plugins/Experimental/EvalGraph/Source/EvalGraphEditor/Private/EvalGraph/EvalGraphSNodeFactories.cpp b/Engine/Plugins/Experimental/EvalGraph/Source/EvalGraphEditor/Private/EvalGraph/EvalGraphSNodeFactories.cpp
new file mode 100644
--- /dev/null
+++ b/Engine/Plugins/Experimental/EvalGraph/Source/EvalGraphEditor/Private/EvalGraph/EvalGraphSNodeFactories.cpp
@@ -0,0 +1,14 @@
+// Copyright Epic Games, Inc. All Rights Reserved.
+
+#include "EvalGraph/EvalGraphSNodeFactories.h"
+#include "EvalGraph/EvalGraphSNode.h"
+#include "EvalGraph/EvalGraphEdNode.h"
+
+TSharedPtr<class SGraphNode> FEvalGraphSNodeFactory::CreateNode(UEdGraphNode* InNode) const
+{
+	if (UEvalGraphEdNode* Node = Cast<UEvalGraphEdNode>(InNode))
+	{
+		return SNew(SEvalGraphEdNode, Node);
+	}
+	return NULL;
+}
diff --git a/Engine/Plugins/Experimental/EvalGraph/Source/EvalGraphEditor/Private/EvalGraph/EvalGraphSchema.cpp b/Engine/Plugins/Experimental/EvalGraph/Source/EvalGraphEditor/Private/EvalGraph/EvalGraphSchema.cpp
new file mode 100644
--- /dev/null
+++ b/Engine/Plugins/Experimental/EvalGraph/Source/EvalGraphEditor/Private/EvalGraph/EvalGraphSchema.cpp
@@ -0,0 +1,55 @@
+// Copyright Epic Games, Inc. All Rights Reserved.
+
+#include "EvalGraph/EvalGraphSchema.h"
+
+#include "EvalGraph/EvalGraphEdNode.h"
+#include "EvalGraph/EvalGraphEditorActions.h"
+#include "EvalGraph/EvalGraphEdNode.h"
+#include "EvalGraph/EvalGraphSNode.h"
+#include "EvalGraph/EvalGraphNodeFactory.h"
+#include "EdGraph/EdGraphNode.h"
+#include "EdGraph/EdGraph.h"
+#include "Framework/Commands/GenericCommands.h"
+#include "Logging/LogMacros.h"
+#include "ToolMenuSection.h"
+#include "ToolMenu.h"
+#include "GraphEditorActions.h"
+
+#define LOCTEXT_NAMESPACE "EvalGraphNode"
+
+UEvalGraphSchema::UEvalGraphSchema()
+{
+}
+
+void UEvalGraphSchema::GetContextMenuActions(class UToolMenu* Menu, class UGraphNodeContextMenuContext* Context) const
+{
+	if (Context->Node)
+	{
+		FToolMenuSection& Section = Menu->AddSection("TestGraphSchemaNodeActions", LOCTEXT("ClassActionsMenuHeader", "Node Actions"));
+		{
+			Section.AddMenuEntry(FGenericCommands::Get().Delete);
+			Section.AddMenuEntry(FGenericCommands::Get().Cut);
+			Section.AddMenuEntry(FGenericCommands::Get().Copy);
+			Section.AddMenuEntry(FGenericCommands::Get().Duplicate);
+			Section.AddMenuEntry(FGraphEditorCommands::Get().BreakNodeLinks);
+			Section.AddMenuEntry(FEvalGraphEditorCommands::Get().EvaluateNode);
+		}
+	}
+
+	Super::GetContextMenuActions(Menu, Context);
+}
+
+void UEvalGraphSchema::GetGraphContextActions(FGraphContextMenuBuilder& ContextMenuBuilder) const
+{
+	if (Eg::FNodeFactory* Factory = Eg::FNodeFactory::GetInstance())
+	{
+		for (FName NodeName : Factory->RegisteredNodes())
+		{
+			if (FEvalGraphEditorCommands::Get().CreateNodesMap.Contains(NodeName))
+			{
+				ContextMenuBuilder.AddAction(FAssetSchemaAction_EvalGraph_CreateNode_EvalGraphEdNode::CreateAction(ContextMenuBuilder.OwnerOfTemporaries, NodeName));
+			}
+		}
+	}
+}
+#undef LOCTEXT_NAMESPACE
diff --git a/Engine/Plugins/Experimental/EvalGraph/Source/EvalGraphEditor/Public/EvalGraph/EvalGraphEditorActions.h b/Engine/Plugins/Experimental/EvalGraph/Source/EvalGraphEditor/Public/EvalGraph/EvalGraphEditorActions.h
--- a/Engine/Plugins/Experimental/EvalGraph/Source/EvalGraphEditor/Public/EvalGraph/EvalGraphEditorActions.h
+++ b/Engine/Plugins/Experimental/EvalGraph/Source/EvalGraphEditor/Public/EvalGraph/EvalGraphEditorActions.h
@@ -22,6 +22,7 @@ public:
 	EVALGRAPHEDITOR_API virtual void RegisterCommands() override;
 
 	TSharedPtr< FUICommandInfo > EvaluateNode;
+	TMap< FName, TSharedPtr<FUICommandInfo> > CreateNodesMap;
 };
 
 class EVALGRAPHEDITOR_API FEvalGraphEditorCommands
diff --git a/Engine/Plugins/Experimental/EvalGraph/Source/EvalGraphEditor/Public/EvalGraph/EvalGraphEditorPlugin.h b/Engine/Plugins/Experimental/EvalGraph/Source/EvalGraphEditor/Public/EvalGraph/EvalGraphEditorPlugin.h
--- a/Engine/Plugins/Experimental/EvalGraph/Source/EvalGraphEditor/Public/EvalGraph/EvalGraphEditorPlugin.h
+++ b/Engine/Plugins/Experimental/EvalGraph/Source/EvalGraphEditor/Public/EvalGraph/EvalGraphEditorPlugin.h
@@ -9,7 +9,7 @@
 
 //class ISlateStyle;
 //class FSlateStyleSet;
-//class FEvalGraphNodeFactory;
+class FEvalGraphSNodeFactory;
 class FEvalGraphAssetActions;
 
 /**
@@ -21,6 +21,9 @@ public:
 	virtual void StartupModule();
 	virtual void ShutdownModule();
 
+	TSharedRef<FAssetEditorToolkit> CreateEvalGraphAssetEditor(const EToolkitMode::Type Mode, const TSharedPtr< IToolkitHost >& InitToolkitHost, UObject* EvalGraph);
+
+
 	/**
 	 * Singleton-like access to this module's interface.  This is just for convenience!
 	 * Beware of calling this during the shutdown phase, though.  Your module might have been unloaded already.
@@ -29,7 +32,7 @@ public:
 	 */
 	static inline IEvalGraphEditorPlugin& Get()
 	{
-		return FModuleManager::LoadModuleChecked< IEvalGraphEditorPlugin >( "EvalGraphEditorPlugin" );
+		return FModuleManager::LoadModuleChecked< IEvalGraphEditorPlugin >( "EvalGraphEditor" );
 	}
 
 	/**
@@ -39,7 +42,7 @@ public:
 	 */
 	static inline bool IsAvailable()
 	{
-		return FModuleManager::Get().IsModuleLoaded( "EvalGraphEditorPlugin" );
+		return FModuleManager::Get().IsModuleLoaded( "EvalGraphEditor" );
 	}
 
 	//TSharedPtr<FSlateStyleSet> GetStyleSet() { return StyleSet; }
@@ -47,14 +50,14 @@ public:
 	//static const ISlateStyle* GetEditorStyle();
 
 private:
-	//void RegisterMenus();
+	void RegisterMenus();
 
 private:
 
 	FEvalGraphAssetActions* EvalGraphAssetActions;
 
 	//TSharedPtr<FSlateStyleSet> StyleSet;
-	//TSharedPtr<FEvalGraphNodeFactory> EvalGraphNodeFactory;
+	TSharedPtr<FEvalGraphSNodeFactory> EvalGraphSNodeFactory;
 
 };
 
diff --git a/Engine/Plugins/Experimental/EvalGraph/Source/EvalGraphEditor/Public/EvalGraph/EvalGraphNodeFactories.h b/Engine/Plugins/Experimental/EvalGraph/Source/EvalGraphEditor/Public/EvalGraph/EvalGraphNodeFactories.h
deleted file mode 100644
diff --git a/Engine/Plugins/Experimental/EvalGraph/Source/EvalGraphEditor/Public/EvalGraph/EvalGraphSNode.h b/Engine/Plugins/Experimental/EvalGraph/Source/EvalGraphEditor/Public/EvalGraph/EvalGraphSNode.h
new file mode 100644
--- /dev/null
+++ b/Engine/Plugins/Experimental/EvalGraph/Source/EvalGraphEditor/Public/EvalGraph/EvalGraphSNode.h
@@ -0,0 +1,56 @@
+// Copyright Epic Games, Inc. All Rights Reserved.
+
+#pragma once
+
+#include "EdGraphUtilities.h"
+#include "SGraphNode.h"
+
+#include "EvalGraphSNode.generated.h"
+
+class UEvalGraphEdNode;
+
+//
+// SEvalGraphEdNode
+//
+
+class EVALGRAPHEDITOR_API SEvalGraphEdNode : public SGraphNode
+{
+	typedef SGraphNode Super;
+
+public:
+	SLATE_BEGIN_ARGS(SEvalGraphEdNode)
+		: _GraphNodeObj(nullptr)
+	{}
+
+	SLATE_ARGUMENT(UEvalGraphEdNode*, GraphNodeObj)
+
+	SLATE_END_ARGS()
+	
+	void Construct(const FArguments& InArgs, UEvalGraphEdNode* InNode);
+
+	// SGraphNode interface
+	virtual FReply OnMouseButtonDoubleClick(const FGeometry& InMyGeometry, const FPointerEvent& InMouseEvent) override;
+
+};
+
+//
+// Action to add a node to the graph
+//
+USTRUCT()
+struct EVALGRAPHEDITOR_API FAssetSchemaAction_EvalGraph_CreateNode_EvalGraphEdNode : public FEdGraphSchemaAction
+{
+	GENERATED_USTRUCT_BODY();
+
+public:
+	FAssetSchemaAction_EvalGraph_CreateNode_EvalGraphEdNode() : NodeTemplate(nullptr) {}
+
+	FAssetSchemaAction_EvalGraph_CreateNode_EvalGraphEdNode(const FText& InNodeCategory, const FText& InMenuDesc, const FText& InToolTip, const int32 InGrouping)
+		: FEdGraphSchemaAction(InNodeCategory, InMenuDesc, InToolTip, InGrouping), NodeTemplate(nullptr) {}
+
+	static TSharedPtr<FAssetSchemaAction_EvalGraph_CreateNode_EvalGraphEdNode> CreateAction(UEdGraph* Owner, const FName & NodeTypeName);
+
+	virtual UEdGraphNode* PerformAction(class UEdGraph* ParentGraph, UEdGraphPin* FromPin, const FVector2D Location, bool bSelectNewNode = true) override;
+	virtual void AddReferencedObjects(FReferenceCollector& Collector) override;
+
+	UEvalGraphEdNode* NodeTemplate;
+};
diff --git a/Engine/Plugins/Experimental/EvalGraph/Source/EvalGraphEditor/Public/EvalGraph/EvalGraphSNodeFactories.h b/Engine/Plugins/Experimental/EvalGraph/Source/EvalGraphEditor/Public/EvalGraph/EvalGraphSNodeFactories.h
new file mode 100644
--- /dev/null
+++ b/Engine/Plugins/Experimental/EvalGraph/Source/EvalGraphEditor/Public/EvalGraph/EvalGraphSNodeFactories.h
@@ -0,0 +1,12 @@
+// Copyright Epic Games, Inc. All Rights Reserved.
+
+#pragma once
+
+#include "EdGraphUtilities.h"
+#include "SGraphNode.h"
+#include "EdGraph/EdGraphNode.h"
+
+class EVALGRAPHEDITOR_API FEvalGraphSNodeFactory : public FGraphPanelNodeFactory
+{
+	virtual TSharedPtr<class SGraphNode> CreateNode(class UEdGraphNode* InNode) const override;
+};
diff --git a/Engine/Plugins/Experimental/EvalGraph/Source/EvalGraphEditor/Public/EvalGraph/EvalGraphSchema.h b/Engine/Plugins/Experimental/EvalGraph/Source/EvalGraphEditor/Public/EvalGraph/EvalGraphSchema.h
new file mode 100644
--- /dev/null
+++ b/Engine/Plugins/Experimental/EvalGraph/Source/EvalGraphEditor/Public/EvalGraph/EvalGraphSchema.h
@@ -0,0 +1,20 @@
+// Copyright Epic Games, Inc. All Rights Reserved.
+
+#pragma once
+
+#include "EvalGraphSchema.generated.h"
+
+class UEvalGraph; 
+
+UCLASS()
+class EVALGRAPHEDITOR_API UEvalGraphSchema : public UEdGraphSchema
+{
+	GENERATED_BODY()
+public:
+	UEvalGraphSchema();
+
+	//~ Begin EdGraphSchema Interface
+	virtual void GetContextMenuActions(class UToolMenu* Menu, class UGraphNodeContextMenuContext* Context) const override;
+	virtual void GetGraphContextActions(FGraphContextMenuBuilder& ContextMenuBuilder) const override;
+	//~ End EdGraphSchema Interface
+};
diff --git a/Engine/Source/Runtime/Experimental/EvalGraph/Core/Private/EvalGraph/EvalGraph.cpp b/Engine/Source/Runtime/Experimental/EvalGraph/Core/Private/EvalGraph/EvalGraph.cpp
--- a/Engine/Source/Runtime/Experimental/EvalGraph/Core/Private/EvalGraph/EvalGraph.cpp
+++ b/Engine/Source/Runtime/Experimental/EvalGraph/Core/Private/EvalGraph/EvalGraph.cpp
@@ -20,11 +20,11 @@ namespace Eg
 
 	void FGraph::RemoveNode(TSharedPtr<FNode> Node)
 	{
-		for (FConnectionTypeBase* Output : Node->GetOutputs())
+		for (FConnectionBase* Output : Node->GetOutputs())
 		{
 			if (Output)
 			{
-				for (FConnectionTypeBase* Input : Output->GetBaseInputs())
+				for (FConnectionBase* Input : Output->GetBaseInputs())
 				{
 					if (Input)
 					{
@@ -33,12 +33,12 @@ namespace Eg
 				}
 			}
 		}
-		for (FConnectionTypeBase* Input : Node->GetInputs())
+		for (FConnectionBase* Input : Node->GetInputs())
 		{
 			if (Input)
 			{
-				TArray<FConnectionTypeBase*> Outputs = Input->GetBaseOutputs();
-				for (FConnectionTypeBase* Out : Outputs)
+				TArray<FConnectionBase*> Outputs = Input->GetBaseOutputs();
+				for (FConnectionBase* Out : Outputs)
 				{
 					if (Out)
 					{
@@ -50,7 +50,7 @@ namespace Eg
 		Nodes.Remove(Node);
 	}
 
-	void FGraph::Connect(FConnectionTypeBase* Input, FConnectionTypeBase* Output)
+	void FGraph::Connect(FConnectionBase* Input, FConnectionBase* Output)
 	{
 		if (ensure(Input && Output))
 		{
@@ -60,7 +60,7 @@ namespace Eg
 		}
 	}
 
-	void FGraph::Disconnect(FConnectionTypeBase* Input, FConnectionTypeBase* Output)
+	void FGraph::Disconnect(FConnectionBase* Input, FConnectionBase* Output)
 	{
 		Input->RemoveConnection(Output);
 		Output->RemoveConnection(Input);
@@ -84,10 +84,10 @@ namespace Eg
 				Ar << ArGuid << ArType << ArName;
 
 
-				TArray< FConnectionTypeBase* > IO = Node->GetOutputs();  IO.Append(Node->GetInputs());
+				TArray< FConnectionBase* > IO = Node->GetOutputs();  IO.Append(Node->GetInputs());
 				ArNum = IO.Num();
 				Ar << ArNum;
-				for (FConnectionTypeBase* Conn : IO)
+				for (FConnectionBase* Conn : IO)
 				{
 					ArGuid = Conn->GetGuid();
 					ArType = Conn->GetType();
@@ -107,7 +107,7 @@ namespace Eg
 			FName ArType, ArName;
 			int32 ArNum = 0;
 
-			TMap<FGuid, FConnectionTypeBase* > ConnectionGuidMap;
+			TMap<FGuid, FConnectionBase* > ConnectionGuidMap;
 
 			Ar << ArNum;
 			for (int32 Ndx = ArNum; Ndx > 0; Ndx--)
@@ -117,7 +117,7 @@ namespace Eg
 				{
 					int ArNumInner;
 					Ar << ArNumInner;
-					TArray< FConnectionTypeBase* > IO = Node->GetOutputs();  IO.Append(Node->GetInputs());
+					TArray< FConnectionBase* > IO = Node->GetOutputs();  IO.Append(Node->GetInputs());
 					for (int Cdx=0;Cdx< ArNumInner; Cdx++)
 					{
 						Ar << ArGuid << ArType << ArName;
@@ -158,6 +158,5 @@ namespace Eg
 
 		}
 	}
-
 }
 
diff --git a/Engine/Source/Runtime/Experimental/EvalGraph/Core/Private/EvalGraph/EvalGraphConnectionBase.cpp b/Engine/Source/Runtime/Experimental/EvalGraph/Core/Private/EvalGraph/EvalGraphConnectionBase.cpp
new file mode 100644
--- /dev/null
+++ b/Engine/Source/Runtime/Experimental/EvalGraph/Core/Private/EvalGraph/EvalGraphConnectionBase.cpp
@@ -0,0 +1,23 @@
+// Copyright Epic Games, Inc. All Rights Reserved.
+
+#include "EvalGraph/EvalGraphConnectionBase.h"
+
+#include "EvalGraph/EvalGraphNodeParameters.h"
+#include "EvalGraph/EvalGraphNode.h"
+
+namespace Eg
+{
+	FConnectionBase::FConnectionBase(FName InType, FName InName, FNode* InOwningNode, FGuid InGuid)
+		: Type(InType)
+		, Name(InName)
+		, Guid(InGuid)
+		, OwningNode(InOwningNode)
+	{}
+
+	void FConnectionBase::AddBaseInput(FNode* InNode, FConnectionBase* That) 
+	{ InNode->AddBaseInput(That); }
+	void FConnectionBase::AddBaseOutput(FNode* InNode, FConnectionBase* That) 
+	{ InNode->AddBaseOutput(That); }
+
+}
+
diff --git a/Engine/Source/Runtime/Experimental/EvalGraph/Core/Private/EvalGraph/EvalGraphInputOutput.cpp b/Engine/Source/Runtime/Experimental/EvalGraph/Core/Private/EvalGraph/EvalGraphInputOutput.cpp
--- a/Engine/Source/Runtime/Experimental/EvalGraph/Core/Private/EvalGraph/EvalGraphInputOutput.cpp
+++ b/Engine/Source/Runtime/Experimental/EvalGraph/Core/Private/EvalGraph/EvalGraphInputOutput.cpp
@@ -7,92 +7,12 @@
 
 namespace Eg
 {
-	FConnectionTypeBase::FConnectionTypeBase(FName InType, FName InName, FNode* InOwningNode, FGuid InGuid)
-		: Type(InType)
-		, Name(InName)
-		, Guid(InGuid)
-		, OwningNode(InOwningNode)
-	{}
-
-	void FConnectionTypeBase::AddBaseInput(FNode* InNode, FConnectionTypeBase* That) 
-	{ InNode->AddBaseInput(That); }
-	void FConnectionTypeBase::AddBaseOutput(FNode* InNode, FConnectionTypeBase* That) 
-	{ InNode->AddBaseOutput(That); }
-
-
-	template<class T>
-	FInput<T>::FInput(const FInputParameters<T>& Param, FGuid InGuid )
-		: FConnectionTypeBase(Param.Type, Param.Name, Param.Owner, InGuid)
-		, Default(Param.Default)
-		, Connection(nullptr)
-	{
-		Super::AddBaseInput(Param.Owner, this);
-	}
-
-	template<class T>
-	void FInput<T>::SetValue(const T& Value, const FContext& Context)
-	{
-		Default = Value;
-		if (!Connection)
-		{
-			OwningNode->InvalidateOutputs();
-		}
-	}
-
-	template<class T>
-	void FInput<T>::Invalidate()
-	{
-		OwningNode->InvalidateOutputs();
-	}
-
-	template<class T>
-	FOutput<T>::FOutput(const FOutputParameters<T>& Param, FGuid InGuid )
-		: FConnectionTypeBase(Param.Type, Param.Name, Param.Owner, InGuid)
-	{
-		Super::AddBaseOutput(Param.Owner, this);
-	}
-
-
-	template<class T>
-	void FOutput<T>::SetValue(T InVal, const FContext& Context)
-	{
-		CacheKey = Context.GetTypeHash();
-		Cache.Data = InVal;
-	}
-
-	template<class T>
-	T FOutput<T>::Evaluate(const FContext& Context)
-	{
-		if (CacheKey!=Context.GetTypeHash())
-		{
-			OwningNode->Evaluate(Context, this);
-		}
-		ensure(CacheKey==Context.GetTypeHash());
-		return Cache.Data;
-	}
-
-	template<class T>
-	void FOutput<T>::Invalidate( )
-	{
-		if (CacheKey != UINT_MAX)
-		{
-			CacheKey = UINT_MAX;
-			Cache = T();
-
-			for (FConnectionTypeBase* Con : GetConnections())
-			{
-				Con->Invalidate();
-			}
-		}
-	}
-
-// Base Types
-EVAL_GRAPH_CONNECTION_TYPE(bool, Bool)
-EVAL_GRAPH_CONNECTION_TYPE(char, Char)
-EVAL_GRAPH_CONNECTION_TYPE(int, Integer)
-EVAL_GRAPH_CONNECTION_TYPE(uint8, UInt8)
-EVAL_GRAPH_CONNECTION_TYPE(float, Float)
-EVAL_GRAPH_CONNECTION_TYPE(double, Double)
-
+	// Base Types
+	EVAL_GRAPH_CONNECTION_TYPE(bool, Bool)
+	EVAL_GRAPH_CONNECTION_TYPE(char, Char)
+	EVAL_GRAPH_CONNECTION_TYPE(int, Integer)	
+	EVAL_GRAPH_CONNECTION_TYPE(uint8, UInt8)
+	EVAL_GRAPH_CONNECTION_TYPE(float, Float)
+	EVAL_GRAPH_CONNECTION_TYPE(double, Double)
 }
 
diff --git a/Engine/Source/Runtime/Experimental/EvalGraph/Core/Private/EvalGraph/EvalGraphNode.cpp b/Engine/Source/Runtime/Experimental/EvalGraph/Core/Private/EvalGraph/EvalGraphNode.cpp
--- a/Engine/Source/Runtime/Experimental/EvalGraph/Core/Private/EvalGraph/EvalGraphNode.cpp
+++ b/Engine/Source/Runtime/Experimental/EvalGraph/Core/Private/EvalGraph/EvalGraphNode.cpp
@@ -6,9 +6,19 @@
 
 namespace Eg
 {
+	TArray<FPin> FNode::GetPins() const
+	{
+		TArray<FPin> RetVal;
+		for (FConnectionBase* Con : Inputs)
+			RetVal.Add({ FPin::EDirection::INPUT,Con->GetType(), Con->GetName() });
+		for (FConnectionBase* Con : Outputs)
+			RetVal.Add({ FPin::EDirection::OUTPUT,Con->GetType(), Con->GetName() });
+		return RetVal;
+	}
+
 	void FNode::InvalidateOutputs()
 	{
-		for(FConnectionTypeBase * Output :  Outputs)
+		for(FConnectionBase * Output :  Outputs)
 		{
 			Output->Invalidate();
 		}
diff --git a/Engine/Source/Runtime/Experimental/EvalGraph/Core/Public/EvalGraph/EvalGraph.h b/Engine/Source/Runtime/Experimental/EvalGraph/Core/Public/EvalGraph/EvalGraph.h
--- a/Engine/Source/Runtime/Experimental/EvalGraph/Core/Public/EvalGraph/EvalGraph.h
+++ b/Engine/Source/Runtime/Experimental/EvalGraph/Core/Public/EvalGraph/EvalGraph.h
@@ -4,12 +4,12 @@
 
 #include "CoreMinimal.h"
 #include "Chaos/ChaosArchive.h"
-#include "Serialization/Archive.h"
 #include "EvalGraph/EvalGraphNode.h"
+#include "Serialization/Archive.h"
 
 namespace Eg
 {
-	class FConnectionTypeBase;
+	class FConnectionBase;
 
 	typedef TTuple<FGuid, FGuid> FConnection; // <Input, Output>
 
@@ -23,7 +23,7 @@ namespace Eg
 		FGuid  Guid;
 		TArray< TSharedPtr<FNode> > Nodes;
 		TArray< FConnection > Connections;
-		//TMap<FGuid, TSharedPtr<FConnectionTypeBase> > ConnectionMap;
+		//TMap<FGuid, TSharedPtr<FConnectionBase> > ConnectionMap;
 
 	public:
 		FGraph(FGuid InGuid = FGuid::NewGuid());
@@ -40,11 +40,11 @@ namespace Eg
 
 		template<class T> T* FindNode(FName InName)
 		{
-			for (TSharedPtr<FNode> Node : Nodes)  
+			for (TSharedPtr<FNode> Node : Nodes)
 			{
-				if (Node->GetName().ToString().Equals(InName.ToString())) 
+				if (Node->GetName().ToString().Equals(InName.ToString()))
 				{
-					if (Node->GetType().ToString().Equals(T::Type.ToString())) 
+					if (Node->GetType().ToString().Equals(T::Type.ToString()))
 					{
 						return reinterpret_cast<T*>(Node.Get()); // @todo(eg) : Can we do better here?
 					}
@@ -53,16 +53,27 @@ namespace Eg
 			return nullptr;
 		}
 
+		TSharedPtr<FNode> FindBaseNode(FGuid InGuid)
+		{
+			for (TSharedPtr<FNode> Node : Nodes)
+			{
+				if (Node->GetGuid() == InGuid)
+				{
+					return Node;
+				}
+			}
+			return TSharedPtr<FNode>(nullptr);
+		}
+
 
 		void RemoveNode(TSharedPtr<FNode> Node);
 
-		void Connect(FConnectionTypeBase* Input, FConnectionTypeBase* Output);
-		void Disconnect(FConnectionTypeBase* Input, FConnectionTypeBase* Output);
+		void Connect(FConnectionBase* Input, FConnectionBase* Output);
+		void Disconnect(FConnectionBase* Input, FConnectionBase* Output);
 
 		virtual void Serialize(FArchive& Ar);
 
 	};
-
 }
 
 FORCEINLINE FArchive& operator<<(FArchive& Ar, Eg::FGraph& Value)
diff --git a/Engine/Source/Runtime/Experimental/EvalGraph/Core/Public/EvalGraph/EvalGraphConnectionBase.h b/Engine/Source/Runtime/Experimental/EvalGraph/Core/Public/EvalGraph/EvalGraphConnectionBase.h
new file mode 100644
--- /dev/null
+++ b/Engine/Source/Runtime/Experimental/EvalGraph/Core/Public/EvalGraph/EvalGraphConnectionBase.h
@@ -0,0 +1,55 @@
+// Copyright Epic Games, Inc. All Rights Reserved.
+
+#pragma once
+
+#include "CoreMinimal.h"
+#include "EvalGraph/EvalGraphNodeParameters.h"
+
+namespace Eg
+{
+	class FNode;
+
+	template<class T> inline EVALGRAPHCORE_API FName GraphConnectionTypeName();
+
+
+	//
+	// Input Output Base
+	//
+
+	class EVALGRAPHCORE_API FConnectionBase
+	{
+
+	protected:
+		FName Type;
+		FName Name;
+		FGuid  Guid;
+		FNode* OwningNode = nullptr;
+
+		friend class FNode;
+		static void AddBaseInput(FNode* InNode, FConnectionBase*);
+		static void AddBaseOutput(FNode* InNode, FConnectionBase*);
+
+
+	public:
+		FConnectionBase(FName InType, FName InName, FNode* OwningNode = nullptr, FGuid InGuid = FGuid::NewGuid());
+		virtual ~FConnectionBase() {};
+
+		FName GetType() const { return Type; }
+
+		FGuid GetGuid() const { return Guid; }
+		void SetGuid(FGuid InGuid) { Guid = InGuid; }
+
+		FName GetName() const { return Name; }
+		void SetName(FName InName) { Name = InName; }
+
+
+		virtual bool AddConnection(FConnectionBase* In) { return false; };
+		virtual bool RemoveConnection(FConnectionBase* In) { return false; }
+
+		virtual TArray< FConnectionBase* > GetBaseInputs() { return TArray<FConnectionBase* >(); }
+		virtual TArray< FConnectionBase* > GetBaseOutputs() { return TArray<FConnectionBase* >(); }
+		virtual void Invalidate() {};
+
+	};
+}
+
diff --git a/Engine/Source/Runtime/Experimental/EvalGraph/Core/Public/EvalGraph/EvalGraphInputOutput.h b/Engine/Source/Runtime/Experimental/EvalGraph/Core/Public/EvalGraph/EvalGraphInputOutput.h
--- a/Engine/Source/Runtime/Experimental/EvalGraph/Core/Public/EvalGraph/EvalGraphInputOutput.h
+++ b/Engine/Source/Runtime/Experimental/EvalGraph/Core/Public/EvalGraph/EvalGraphInputOutput.h
@@ -4,54 +4,11 @@
 
 #include "CoreMinimal.h"
 #include "EvalGraph/EvalGraphNodeParameters.h"
+#include "EvalGraph/EvalGraphNode.h"
+#include "EvalGraph/EvalGraphConnectionBase.h"
 
 namespace Eg
 {
-	class FNode;
-
-	template<class T> inline EVALGRAPHCORE_API FName GraphConnectionTypeName();
-
-
-	//
-	// Input Output Base
-	//
-
-	class EVALGRAPHCORE_API FConnectionTypeBase
-	{
-
-	protected:
-		FName Type;
-		FName Name;
-		FGuid  Guid;
-		FNode* OwningNode = nullptr;
-
-		friend class FNode;
-		static void AddBaseInput(FNode* InNode, FConnectionTypeBase*);
-		static void AddBaseOutput(FNode* InNode, FConnectionTypeBase*);
-
-
-	public:
-		FConnectionTypeBase(FName InType, FName InName, FNode* OwningNode = nullptr, FGuid InGuid = FGuid::NewGuid());
-		virtual ~FConnectionTypeBase() {};
-
-		FName GetType() const { return Type; }
-
-		FGuid GetGuid() const { return Guid; }
-		void SetGuid(FGuid InGuid) { Guid = InGuid; }
-
-		FName GetName() const { return Name; }
-		void SetName(FName InName) { Name = InName; }
-
-
-		virtual bool AddConnection(FConnectionTypeBase* In) { return false; };
-		virtual bool RemoveConnection(FConnectionTypeBase* In) { return false; }
-
-		virtual TArray< FConnectionTypeBase* > GetBaseInputs() { return TArray<FConnectionTypeBase* >(); }
-		virtual TArray< FConnectionTypeBase* > GetBaseOutputs() { return TArray<FConnectionTypeBase* >(); }
-		virtual void Invalidate() {};
-
-	};
-
 	template<typename T> class FOutput;
 
 	//
@@ -71,22 +28,28 @@ namespace Eg
 	};
 
 	template<class T>
-	class EVALGRAPHCORE_API FInput : public FConnectionTypeBase
+	class EVALGRAPHCORE_API FInput : public FConnectionBase
 	{
-		friend class FConnectionTypeBase;
+		typedef FConnectionBase Super;
+		friend class FConnectionBase;
 
-		typedef FConnectionTypeBase Super;
 		T Default;
 		FOutput<T>* Connection;
 	public:
-		FInput(const FInputParameters<T>& Param, FGuid InGuid = FGuid::NewGuid());
+		FInput(const FInputParameters<T>& Param, FGuid InGuid = FGuid::NewGuid())
+			: FConnectionBase(Param.Type, Param.Name, Param.Owner, InGuid)
+			, Default(Param.Default)
+			, Connection(nullptr)
+		{
+			Super::AddBaseInput(Param.Owner, this);
+		}
 	
 		const T& GetDefault() const { return Default; }
 		
 		const FOutput<T>* GetConnection() const { return Connection; }
 		FOutput<T>* GetConnection() { return Connection; }
 
-		virtual bool AddConnection(FConnectionTypeBase* InOutput) override
+		virtual bool AddConnection(FConnectionBase* InOutput) override
 		{ 
 			ensure(Connection == nullptr);
 			if (ensure(InOutput->GetType()==this->GetType()))
@@ -97,7 +60,7 @@ namespace Eg
 			return false;
 		}
 
-		virtual bool RemoveConnection(FConnectionTypeBase* InOutput) override
+		virtual bool RemoveConnection(FConnectionBase* InOutput) override
 		{ 
 			if (ensure(Connection == (FOutput<T>*)InOutput))
 			{
@@ -116,11 +79,18 @@ namespace Eg
 			return Default;
 		}
 
-		void SetValue(const T& Value, const FContext& Context);
+		void SetValue(const T& Value, const FContext& Context)
+		{
+			Default = Value;
+			if (!Connection)
+			{
+				OwningNode->InvalidateOutputs();
+			}
+		}
 
-		virtual TArray< FConnectionTypeBase* > GetBaseOutputs() override
+		virtual TArray< FConnectionBase* > GetBaseOutputs() override
 		{ 
-			TArray<FConnectionTypeBase* > RetList;
+			TArray<FConnectionBase* > RetList;
 			if (GetConnection())
 			{
 				RetList.Add(GetConnection());
@@ -128,7 +98,10 @@ namespace Eg
 			return RetList;
 		}
 
-		virtual void Invalidate() override;
+		virtual void Invalidate() override
+		{
+			OwningNode->InvalidateOutputs();
+		}
 
 	};
 
@@ -149,31 +122,36 @@ namespace Eg
 	};
 
 	template<class T>
-	class EVALGRAPHCORE_API FOutput : public FConnectionTypeBase
+	class EVALGRAPHCORE_API FOutput : public FConnectionBase
 	{
-		typedef FConnectionTypeBase Super;
-		friend class FConnectionTypeBase;
+		typedef FConnectionBase Super;
+		friend class FConnectionBase;
 
 		uint32 CacheKey = UINT_MAX;
 		TCacheValue<T> Cache;
 		TArray< FInput<T>* > Connections;
 	
 	public:
-		FOutput(const FOutputParameters<T>& Param, FGuid InGuid = FGuid::NewGuid());
+		FOutput(const FOutputParameters<T>& Param, FGuid InGuid = FGuid::NewGuid())
+			: FConnectionBase(Param.Type, Param.Name, Param.Owner, InGuid)
+		{
+			Super::AddBaseOutput(Param.Owner, this);
+		}
+
 
 		const TArray<FInput<T>*>& GetConnections() const { return Connections; }
 		TArray<FInput<T>* >& GetConnections() { return Connections; }
 
-		virtual TArray<FConnectionTypeBase*> GetBaseInputs() override
+		virtual TArray<FConnectionBase*> GetBaseInputs() override
 		{ 
-			TArray<FConnectionTypeBase*> RetList;
+			TArray<FConnectionBase*> RetList;
 			RetList.Reserve(Connections.Num());
 			for (FInput<T>* Ptr : Connections) { RetList.Add(Ptr); }
 			return RetList;
 		}
 
 
-		virtual bool AddConnection(FConnectionTypeBase* InOutput) override
+		virtual bool AddConnection(FConnectionBase* InOutput) override
 		{
 			if (ensure(InOutput->GetType() == this->GetType()))
 			{
@@ -183,13 +161,37 @@ namespace Eg
 			return false;
 		}
 		
-		virtual bool RemoveConnection(FConnectionTypeBase* InInput) override { Connections.RemoveSwap((FInput<T>*)InInput); return true; }
+		virtual bool RemoveConnection(FConnectionBase* InInput) override { Connections.RemoveSwap((FInput<T>*)InInput); return true; }
 
-		void SetValue(T InVal, const FContext& Context);
+		void SetValue(T InVal, const FContext& Context)
+		{
+			CacheKey = Context.GetTypeHash();
+			Cache.Data = InVal;
+		}
 
-		T Evaluate(const FContext& Context);
+		T Evaluate(const FContext& Context)
+		{
+			if (CacheKey != Context.GetTypeHash())
+			{
+				OwningNode->Evaluate(Context, this);
+			}
+			ensure(CacheKey == Context.GetTypeHash());
+			return Cache.Data;
+		}
 
-		virtual void Invalidate() override;
+		virtual void Invalidate() override
+		{
+			if (CacheKey != UINT_MAX)
+			{
+				CacheKey = UINT_MAX;
+				Cache = T();
+
+				for (FConnectionBase* Con : GetConnections())
+				{
+					Con->Invalidate();
+				}
+			}
+		}
 
 	};
 
@@ -198,5 +200,7 @@ namespace Eg
 template<> FName GraphConnectionTypeName<a>() { return FName(TEXT(#A)); }\
 template class FOutput<a>;\
 template class FInput<a>;
+
+
 }
 
diff --git a/Engine/Source/Runtime/Experimental/EvalGraph/Core/Public/EvalGraph/EvalGraphNode.h b/Engine/Source/Runtime/Experimental/EvalGraph/Core/Public/EvalGraph/EvalGraphNode.h
--- a/Engine/Source/Runtime/Experimental/EvalGraph/Core/Public/EvalGraph/EvalGraphNode.h
+++ b/Engine/Source/Runtime/Experimental/EvalGraph/Core/Public/EvalGraph/EvalGraphNode.h
@@ -3,12 +3,9 @@
 #pragma once
 
 #include "CoreMinimal.h"
-#include "EvalGraph/EvalGraphInputOutput.h"
+#include "EvalGraph/EvalGraphConnectionBase.h"
 #include "EvalGraph/EvalGraphNodeParameters.h"
 
-class FConnectionTypeBase;
-
-
 namespace Eg
 {
 
@@ -16,20 +13,31 @@ namespace Eg
 		FName Name;
 	};
 
+	struct FPin
+	{
+		enum class EDirection : uint8 {
+			NONE = 0,
+			INPUT,
+			OUTPUT
+		};
+		EDirection Direction;
+		FName Type;
+		FName Name;
+	};
 
 	//
 	// FNode
 	//
-	class FNode
+	class EVALGRAPHCORE_API FNode
 	{
 		friend class FGraph;
-		friend class FConnectionTypeBase;
+		friend class FConnectionBase;
 
 		FGuid Guid;
 		FName Name;
 
-		TArray< FConnectionTypeBase* > Inputs;
-		TArray< FConnectionTypeBase* > Outputs;
+		TArray< FConnectionBase* > Inputs;
+		TArray< FConnectionBase* > Outputs;
 	public:
 
 		FNode(const FNodeParameters& Param, FGuid InGuid = FGuid::NewGuid())
@@ -45,20 +53,21 @@ namespace Eg
 		virtual FName GetType() const { check(true); return FName("invalid"); }
 		FGuid GetGuid() const { return Guid; }
 
+		TArray<FPin> GetPins() const;
 
-		virtual void Evaluate(const FContext& Context, FConnectionTypeBase*) { ensure(false); }
+		virtual void Evaluate(const FContext& Context, FConnectionBase*) { ensure(false); }
 		void InvalidateOutputs();
 
 		virtual void SerializeInternal(FArchive& Ar) {};
 
 	protected:
-		void AddBaseInput(FConnectionTypeBase* InPtr) { Inputs.Add(InPtr); }
-		const TArray< FConnectionTypeBase* >& GetInputs() const { return Inputs; }
-		TArray< FConnectionTypeBase* >& GetInputs() { return Inputs; }
+		void AddBaseInput(FConnectionBase* InPtr) { Inputs.Add(InPtr); }
+		const TArray< FConnectionBase* >& GetInputs() const { return Inputs; }
+		TArray< FConnectionBase* >& GetInputs() { return Inputs; }
 
-		void AddBaseOutput(FConnectionTypeBase* InPtr) { Outputs.Add(InPtr); }
-		const TArray< FConnectionTypeBase* >& GetOutputs() const { return Outputs; }
-		TArray< FConnectionTypeBase* >& GetOutputs() { return Outputs; }
+		void AddBaseOutput(FConnectionBase* InPtr) { Outputs.Add(InPtr); }
+		const TArray< FConnectionBase* >& GetOutputs() const { return Outputs; }
+		TArray< FConnectionBase* >& GetOutputs() { return Outputs; }
 
 	};
 }
diff --git a/Engine/Source/Runtime/Experimental/EvalGraph/Core/Public/EvalGraph/EvalGraphNodeFactory.h b/Engine/Source/Runtime/Experimental/EvalGraph/Core/Public/EvalGraph/EvalGraphNodeFactory.h
--- a/Engine/Source/Runtime/Experimental/EvalGraph/Core/Public/EvalGraph/EvalGraphNodeFactory.h
+++ b/Engine/Source/Runtime/Experimental/EvalGraph/Core/Public/EvalGraph/EvalGraphNodeFactory.h
@@ -8,7 +8,7 @@
 namespace Eg
 {
 	class FNode;
-	class FConnectionTypeBase;
+	class FConnectionBase;
 
 	struct EVALGRAPHCORE_API FNewNodeParameters {
 		FGuid Guid;
@@ -50,11 +50,17 @@ namespace Eg
 
 		TSharedPtr<FNode> NewNodeFromRegisteredType(FGraph& Graph, const FNewNodeParameters& Param);
 
-		void here() {}
 		template<class T> TSharedPtr<T> NewNode(FGraph& Graph, const FNewNodeParameters& Param)
 		{
 			return Graph.AddNode(new T(Param.Name, Param.Guid));
 		}
+
+		TArray<FName> RegisteredNodes() const
+		{
+			TArray<FName> ReturnNames;
+			for (auto Elem : ClassMap) ReturnNames.Add(Elem.Key);
+			return ReturnNames;
+		}
 	};
 
 }
diff --git a/Engine/Source/Runtime/Experimental/EvalGraph/Engine/Private/EvalGraph/EvalGraphEdNode.cpp b/Engine/Source/Runtime/Experimental/EvalGraph/Engine/Private/EvalGraph/EvalGraphEdNode.cpp
new file mode 100644
--- /dev/null
+++ b/Engine/Source/Runtime/Experimental/EvalGraph/Engine/Private/EvalGraph/EvalGraphEdNode.cpp
@@ -0,0 +1,50 @@
+// Copyright Epic Games, Inc. All Rights Reserved.
+
+#include "EvalGraph/EvalGraphEdNode.h"
+
+#include "EvalGraph/EvalGraphNode.h"
+#include "EvalGraph/EvalGraph.h"
+#include "Logging/LogMacros.h"
+
+#define LOCTEXT_NAMESPACE "EvalGraphEdNode"
+
+DEFINE_LOG_CATEGORY_STATIC(EGEvalGraphNodeLOG, Error, All);
+
+
+void UEvalGraphEdNode::AllocateDefaultPins()
+{
+	UE_LOG(EGEvalGraphNodeLOG, Verbose, TEXT("UEvalGraphEdNode::AllocateDefaultPins()"));
+	// called on node creation from UI. 
+#if WITH_EDITOR && !(UE_BUILD_SHIPPING || UE_BUILD_TEST)
+	if (EgGraph)
+	{
+		if (EgNodeGuid.IsValid())
+		{
+			if (TSharedPtr<Eg::FNode> EgNode = EgGraph->FindBaseNode(EgNodeGuid))
+			{
+				for (const Eg::FPin& Pin : EgNode->GetPins())
+				{
+					if (Pin.Direction == Eg::FPin::EDirection::INPUT)
+					{
+						CreatePin(EEdGraphPinDirection::EGPD_Input, Pin.Type, Pin.Name);
+					}					
+					if (Pin.Direction == Eg::FPin::EDirection::OUTPUT)
+					{
+						CreatePin(EEdGraphPinDirection::EGPD_Output, Pin.Type, Pin.Name);
+					}
+				}
+			}
+		}
+	}
+#endif // WITH_EDITOR && !UE_BUILD_SHIPPING
+}
+
+void UEvalGraphEdNode::Serialize(FArchive& Ar)
+{
+	Super::Serialize(Ar);
+	Ar << EgNodeGuid;
+}
+
+
+#undef LOCTEXT_NAMESPACE
+
diff --git a/Engine/Source/Runtime/Experimental/EvalGraph/Engine/Private/EvalGraph/EvalGraphObject.cpp b/Engine/Source/Runtime/Experimental/EvalGraph/Engine/Private/EvalGraph/EvalGraphObject.cpp
--- a/Engine/Source/Runtime/Experimental/EvalGraph/Engine/Private/EvalGraph/EvalGraphObject.cpp
+++ b/Engine/Source/Runtime/Experimental/EvalGraph/Engine/Private/EvalGraph/EvalGraphObject.cpp
@@ -7,7 +7,7 @@
 #define LOCTEXT_NAMESPACE "UEvalGraph")
 
 UEvalGraph::UEvalGraph(const FObjectInitializer& ObjectInitializer)
-	: UObject(ObjectInitializer)
+	: UEdGraph(ObjectInitializer)
 	, EvalGraph(new Eg::FGraph())
 {}
 
@@ -27,6 +27,7 @@ void UEvalGraph::PostLoad()
 
 void UEvalGraph::Serialize(FArchive& Ar)
 {
+	Super::Serialize(Ar);
 	Ar << *EvalGraph.Get();
 }
 
diff --git a/Engine/Source/Runtime/Experimental/EvalGraph/Engine/Public/EvalGraph/EvalGraphEdNode.h b/Engine/Source/Runtime/Experimental/EvalGraph/Engine/Public/EvalGraph/EvalGraphEdNode.h
new file mode 100644
--- /dev/null
+++ b/Engine/Source/Runtime/Experimental/EvalGraph/Engine/Public/EvalGraph/EvalGraphEdNode.h
@@ -0,0 +1,33 @@
+// Copyright Epic Games, Inc. All Rights Reserved.
+
+#pragma once
+
+#include "CoreMinimal.h"
+#include "EdGraph/EdGraph.h"
+#include "EvalGraph/EvalGraph.h"
+
+#include "EvalGraphEdNode.generated.h"
+
+
+
+UCLASS()
+class EVALGRAPHENGINE_API UEvalGraphEdNode : public UEdGraphNode
+{
+	GENERATED_BODY()
+
+	FGuid EgNodeGuid;
+	TSharedPtr<Eg::FGraph> EgGraph;
+
+public:
+
+	// UEdGraphNode interface
+	virtual void AllocateDefaultPins();
+	// End of UEdGraphNode interface
+
+	void SetEgGraph(TSharedPtr<Eg::FGraph> InEgGraph) { EgGraph = InEgGraph; }
+	void SetEgNode(FGuid InGuid) {EgNodeGuid = InGuid;}
+
+	void Serialize(FArchive& Ar);
+
+};
+
diff --git a/Engine/Source/Runtime/Experimental/EvalGraph/Engine/Public/EvalGraph/EvalGraphObject.h b/Engine/Source/Runtime/Experimental/EvalGraph/Engine/Public/EvalGraph/EvalGraphObject.h
--- a/Engine/Source/Runtime/Experimental/EvalGraph/Engine/Public/EvalGraph/EvalGraphObject.h
+++ b/Engine/Source/Runtime/Experimental/EvalGraph/Engine/Public/EvalGraph/EvalGraphObject.h
@@ -3,6 +3,7 @@
 #pragma once
 
 #include "CoreMinimal.h"
+#include "EdGraph/EdGraph.h"
 #include "EvalGraph/EvalGraph.h"
 
 #include "EvalGraphObject.generated.h"
@@ -14,7 +15,7 @@
 *
 */
 UCLASS(BlueprintType, customconstructor)
-class EVALGRAPHENGINE_API UEvalGraph : public UObject
+class EVALGRAPHENGINE_API UEvalGraph : public UEdGraph
 {
 	GENERATED_UCLASS_BODY()
 
