EvalGraph : Synced Connections and Evaluation

diff --git a/Engine/Plugins/Experimental/EvalGraph/Source/EvalGraphEditor/Private/EvalGraph/EvalGraphEditorToolkit.cpp b/Engine/Plugins/Experimental/EvalGraph/Source/EvalGraphEditor/Private/EvalGraph/EvalGraphEditorToolkit.cpp
--- a/Engine/Plugins/Experimental/EvalGraph/Source/EvalGraphEditor/Private/EvalGraph/EvalGraphEditorToolkit.cpp
+++ b/Engine/Plugins/Experimental/EvalGraph/Source/EvalGraphEditor/Private/EvalGraph/EvalGraphEditorToolkit.cpp
@@ -77,28 +77,24 @@ void FEvalGraphEditorToolkit::EvaluateNode()
 {
 	if (UEvalGraph* Graph = dynamic_cast<UEvalGraph*>(GraphEditor->GetCurrentGraph()))
 	{
-		for (UObject* Ode : GetSelectedNodes())
+		if (TSharedPtr<Eg::FGraph> EgEvalGraph = Graph->GetEvalGraph())
 		{
-			if (UEvalGraphEdNode* Node = dynamic_cast<UEvalGraphEdNode*>(Ode))
+			for (UObject* Ode : GetSelectedNodes())
 			{
-				//UE_LOG(FEvalGraphEditorToolkitLog, Verbose, TEXT("FEvalGraphEditorToolkit::EvaluateNode(%s)"), *Ode->GetFullName());
-				//Graph->Evaluate(Node);
+				if (UEvalGraphEdNode* Node = dynamic_cast<UEvalGraphEdNode*>(Ode))
+				{
+					if (TSharedPtr<Eg::FNode> EgNode = EgEvalGraph->FindBaseNode(Node->GetEgNodeGuid()))
+					{
+						for (Eg::FConnectionBase* NodeOutput : EgNode->GetOutputs())
+						{
+							EgNode->Evaluate({0.f}, NodeOutput);
+						}
+					}
+				}
 			}
 		}
 	}
 }
-void FEvalGraphEditorToolkit::CreateNode(FName RegisteredNode)
-{
-	//UE_LOG(FEvalGraphEditorToolkitLog, Verbose, TEXT("FEvalGraphEditorToolkit::CreateNode"), RegisteredNode);
-	if (UEvalGraph* Graph = dynamic_cast<UEvalGraph*>(GraphEditor->GetCurrentGraph()))
-	{
-		if (UEvalGraph * Asset = GetEvalGraph())
-		{
-			//TSharedPtr<FNewManagedArrayCollectionNode> Node1 = AddNode(new FNewManagedArrayCollectionNode({ FName("Node1") }));
-		}
-	}
-}
-
 
 FGraphPanelSelectionSet FEvalGraphEditorToolkit::GetSelectedNodes() const
 {
diff --git a/Engine/Plugins/Experimental/EvalGraph/Source/EvalGraphEditor/Private/EvalGraph/EvalGraphEditorToolkit.h b/Engine/Plugins/Experimental/EvalGraph/Source/EvalGraphEditor/Private/EvalGraph/EvalGraphEditorToolkit.h
--- a/Engine/Plugins/Experimental/EvalGraph/Source/EvalGraphEditor/Private/EvalGraph/EvalGraphEditorToolkit.h
+++ b/Engine/Plugins/Experimental/EvalGraph/Source/EvalGraphEditor/Private/EvalGraph/EvalGraphEditorToolkit.h
@@ -36,7 +36,6 @@ public:
 	// Graph Editor Operations
 	void OnSelectedNodesChanged(const TSet<UObject*>& NewSelection);
 	void EvaluateNode();
-	void CreateNode(FName RegisteredNode);
 	FGraphPanelSelectionSet GetSelectedNodes() const;
 
 	UEvalGraph* GetEvalGraph() { return EvalGraph; }
diff --git a/Engine/Plugins/Experimental/EvalGraph/Source/EvalGraphEditor/Private/EvalGraph/EvalGraphSNode.cpp b/Engine/Plugins/Experimental/EvalGraph/Source/EvalGraphEditor/Private/EvalGraph/EvalGraphSNode.cpp
--- a/Engine/Plugins/Experimental/EvalGraph/Source/EvalGraphEditor/Private/EvalGraph/EvalGraphSNode.cpp
+++ b/Engine/Plugins/Experimental/EvalGraph/Source/EvalGraphEditor/Private/EvalGraph/EvalGraphSNode.cpp
@@ -17,19 +17,6 @@
 void SEvalGraphEdNode::Construct(const FArguments& InArgs, UEvalGraphEdNode* InNode)
 {
 	GraphNode = InNode;
-	UEvalGraph* Graph = Cast<UEvalGraph>(InNode->GetOuter());
-	if( ensure(Graph) )
-	{
-		if (Eg::FNodeFactory* Factory = Eg::FNodeFactory::GetInstance())
-		{
-			if (TSharedPtr<Eg::FNode> EgNode = Factory->NewNodeFromRegisteredType(*Graph->GetEvalGraph(), { FGuid::NewGuid(),FName(InNode->GetName()),FName(InNode->GetName()) }))
-			{
-				InNode->SetEgGraph(Graph->GetEvalGraph());
-				InNode->SetEgNode(EgNode->GetGuid());
-				InNode->AllocateDefaultPins();
-			}
-		}
-	}
 	UpdateGraphNode();
 }
 
@@ -41,57 +28,67 @@ FReply SEvalGraphEdNode::OnMouseButtonDoubleClick(const FGeometry& InMyGeometry,
 
 
 //
-// FAssetSchemaAction_EvalGraph_CreateNode_Example
+// Add a menu option to create a graph node.
 //
-TSharedPtr<FAssetSchemaAction_EvalGraph_CreateNode_EvalGraphEdNode> FAssetSchemaAction_EvalGraph_CreateNode_EvalGraphEdNode::CreateAction(UEdGraph* Owner, const FName & InNodeTypeName)
+TSharedPtr<FAssetSchemaAction_EvalGraph_CreateNode_EvalGraphEdNode> FAssetSchemaAction_EvalGraph_CreateNode_EvalGraphEdNode::CreateAction(UEdGraph* ParentGraph, const FName & InNodeTypeName)
 {
-	const FText AddToolTip = LOCTEXT("EvalGraphNodeTooltip_Example", "Add a Dataflow node.");
+ 	const FText AddToolTip = LOCTEXT("EvalGraphNodeTooltip_Example", "Add a Dataflow node.");
 	const FText NodeName = FText::FromString(*InNodeTypeName.ToString());
 	const FText Catagory = LOCTEXT("EvalGraphNodeDescription_Example", "Dataflow"); 	
 	TSharedPtr<FAssetSchemaAction_EvalGraph_CreateNode_EvalGraphEdNode> NewNodeAction(
-		new FAssetSchemaAction_EvalGraph_CreateNode_EvalGraphEdNode(Catagory, NodeName, AddToolTip, 0));
-	NewNodeAction->NodeTemplate = NewObject<UEvalGraphEdNode>(Owner, UEvalGraphEdNode::StaticClass(), InNodeTypeName);
-	//NewNodeAction->NodeTemplate->GenericGraphNode = NewObject<UGenericGraphNode>(NewNodeAction->NodeTemplate, Graph->NodeType);
-	//NewNodeAction->NodeTemplate->GenericGraphNode->Graph = Graph;
+		new FAssetSchemaAction_EvalGraph_CreateNode_EvalGraphEdNode(InNodeTypeName, Catagory, NodeName, AddToolTip, 0));
 	return NewNodeAction;
 }
 
-
+//
+//  Created the EdGraph node and bind the guids to the EvalGraph's node. 
+//
 UEdGraphNode* FAssetSchemaAction_EvalGraph_CreateNode_EvalGraphEdNode::PerformAction(class UEdGraph* ParentGraph, UEdGraphPin* FromPin, const FVector2D Location, bool bSelectNewNode)
 {
-	UEdGraphNode* ResultNode = nullptr;
-
-	if (NodeTemplate != nullptr)
+	if (UEvalGraph* EvalGraph = Cast<UEvalGraph>(ParentGraph))
 	{
-		//const FScopedTransaction Transaction(LOCTEXT("EvalGraphNewNode", "Generic Graph Editor: New Node"));
-		ParentGraph->Modify();
-		if (FromPin != nullptr)
-			FromPin->Modify();
+		const FName NodeName(GetMenuDescription().ToString());
+		if (UEvalGraphEdNode* EdNode = NewObject<UEvalGraphEdNode>(EvalGraph, UEvalGraphEdNode::StaticClass(), NodeName))
+		{
+			//const FScopedTransaction Transaction(LOCTEXT("EvalGraphNewNode", "Generic Graph Editor: New Node"));
+			EvalGraph->Modify();
+			if (FromPin != nullptr)
+				FromPin->Modify();
 
-		NodeTemplate->Rename(nullptr, ParentGraph);
-		ParentGraph->AddNode(NodeTemplate, true, bSelectNewNode);
+			EvalGraph->AddNode(EdNode, true, bSelectNewNode);
 
-		NodeTemplate->CreateNewGuid();
-		NodeTemplate->PostPlacedNewNode();
-		//NodeTemplate->AllocateDefaultPins();
-		NodeTemplate->AutowireNewNode(FromPin);
+			EdNode->CreateNewGuid();
+			EdNode->PostPlacedNewNode();
 
-		NodeTemplate->NodePosX = Location.X;
-		NodeTemplate->NodePosY = Location.Y;
+			if (Eg::FNodeFactory* Factory = Eg::FNodeFactory::GetInstance())
+			{
+				if (TSharedPtr<Eg::FNode> EgNode = Factory->NewNodeFromRegisteredType(*EvalGraph->GetEvalGraph(), { FGuid::NewGuid(),NodeTypeName,NodeName}))
+				{
+					EdNode->SetEgGraph(EvalGraph->GetEvalGraph());
+					EdNode->SetEgNodeGuid(EgNode->GetGuid());
+					EdNode->AllocateDefaultPins();
+				}
+			}
 
-		//NodeTemplate->GenericGraphNode->SetFlags(RF_Transactional);
-		NodeTemplate->SetFlags(RF_Transactional);
+			//EdNode->AllocateDefaultPins();
+			EdNode->AutowireNewNode(FromPin);
 
-		ResultNode = NodeTemplate;
-	}
+			EdNode->NodePosX = Location.X;
+			EdNode->NodePosY = Location.Y;
 
-	return ResultNode;
-}
+			//EdNode->GenericGraphNode->SetFlags(RF_Transactional);
+			EdNode->SetFlags(RF_Transactional);
 
-void FAssetSchemaAction_EvalGraph_CreateNode_EvalGraphEdNode::AddReferencedObjects(FReferenceCollector& Collector)
-{
-	FEdGraphSchemaAction::AddReferencedObjects(Collector);
-	Collector.AddReferencedObject(NodeTemplate);
+			return EdNode;
+		}
+	}
+	return nullptr;
 }
 
+//void FAssetSchemaAction_EvalGraph_CreateNode_EvalGraphEdNode::AddReferencedObjects(FReferenceCollector& Collector)
+//{
+//	FEdGraphSchemaAction::AddReferencedObjects(Collector);
+//	Collector.AddReferencedObject(NodeTemplate);
+//}
+
 #undef LOCTEXT_NAMESPACE
diff --git a/Engine/Plugins/Experimental/EvalGraph/Source/EvalGraphEditor/Private/EvalGraph/EvalGraphSchema.cpp b/Engine/Plugins/Experimental/EvalGraph/Source/EvalGraphEditor/Private/EvalGraph/EvalGraphSchema.cpp
--- a/Engine/Plugins/Experimental/EvalGraph/Source/EvalGraphEditor/Private/EvalGraph/EvalGraphSchema.cpp
+++ b/Engine/Plugins/Experimental/EvalGraph/Source/EvalGraphEditor/Private/EvalGraph/EvalGraphSchema.cpp
@@ -43,13 +43,30 @@ void UEvalGraphSchema::GetGraphContextActions(FGraphContextMenuBuilder& ContextM
 {
 	if (Eg::FNodeFactory* Factory = Eg::FNodeFactory::GetInstance())
 	{
-		for (FName NodeName : Factory->RegisteredNodes())
+		for (FName NodeTypeName : Factory->RegisteredNodes())
 		{
-			if (FEvalGraphEditorCommands::Get().CreateNodesMap.Contains(NodeName))
+			if (FEvalGraphEditorCommands::Get().CreateNodesMap.Contains(NodeTypeName))
 			{
-				ContextMenuBuilder.AddAction(FAssetSchemaAction_EvalGraph_CreateNode_EvalGraphEdNode::CreateAction(ContextMenuBuilder.OwnerOfTemporaries, NodeName));
+				ContextMenuBuilder.AddAction(FAssetSchemaAction_EvalGraph_CreateNode_EvalGraphEdNode::CreateAction(ContextMenuBuilder.OwnerOfTemporaries, NodeTypeName));
 			}
 		}
 	}
 }
+
+const FPinConnectionResponse UEvalGraphSchema::CanCreateConnection(const UEdGraphPin* PinA, const UEdGraphPin* PinB) const
+{
+	// Make sure the pins are not on the same node
+	if (PinA->GetOwningNode() != PinB->GetOwningNode())
+	{
+		// Make sure types match. 
+		if (PinA->PinType == PinB->PinType)
+		{
+			return FPinConnectionResponse(CONNECT_RESPONSE_MAKE, LOCTEXT("PinConnect", "Connect nodes"));
+		}
+	}
+
+	return FPinConnectionResponse(CONNECT_RESPONSE_DISALLOW, LOCTEXT("PinErrorSameNode", "Both are on the same node"));
+}
+
+
 #undef LOCTEXT_NAMESPACE
diff --git a/Engine/Plugins/Experimental/EvalGraph/Source/EvalGraphEditor/Public/EvalGraph/EvalGraphSNode.h b/Engine/Plugins/Experimental/EvalGraph/Source/EvalGraphEditor/Public/EvalGraph/EvalGraphSNode.h
--- a/Engine/Plugins/Experimental/EvalGraph/Source/EvalGraphEditor/Public/EvalGraph/EvalGraphSNode.h
+++ b/Engine/Plugins/Experimental/EvalGraph/Source/EvalGraphEditor/Public/EvalGraph/EvalGraphSNode.h
@@ -42,15 +42,15 @@ struct EVALGRAPHEDITOR_API FAssetSchemaAction_EvalGraph_CreateNode_EvalGraphEdNo
 	GENERATED_USTRUCT_BODY();
 
 public:
-	FAssetSchemaAction_EvalGraph_CreateNode_EvalGraphEdNode() : NodeTemplate(nullptr) {}
+	FAssetSchemaAction_EvalGraph_CreateNode_EvalGraphEdNode() : FEdGraphSchemaAction() {}
 
-	FAssetSchemaAction_EvalGraph_CreateNode_EvalGraphEdNode(const FText& InNodeCategory, const FText& InMenuDesc, const FText& InToolTip, const int32 InGrouping)
-		: FEdGraphSchemaAction(InNodeCategory, InMenuDesc, InToolTip, InGrouping), NodeTemplate(nullptr) {}
+	FAssetSchemaAction_EvalGraph_CreateNode_EvalGraphEdNode(const FName& InType, const FText& InNodeCategory, const FText& InMenuDesc, const FText& InToolTip, const int32 InGrouping)
+		: FEdGraphSchemaAction(InNodeCategory, InMenuDesc, InToolTip, InGrouping), NodeTypeName(InType) {}
 
-	static TSharedPtr<FAssetSchemaAction_EvalGraph_CreateNode_EvalGraphEdNode> CreateAction(UEdGraph* Owner, const FName & NodeTypeName);
+	static TSharedPtr<FAssetSchemaAction_EvalGraph_CreateNode_EvalGraphEdNode> CreateAction(UEdGraph* ParentGraph, const FName & NodeTypeName);
 
 	virtual UEdGraphNode* PerformAction(class UEdGraph* ParentGraph, UEdGraphPin* FromPin, const FVector2D Location, bool bSelectNewNode = true) override;
-	virtual void AddReferencedObjects(FReferenceCollector& Collector) override;
+	//virtual void AddReferencedObjects(FReferenceCollector& Collector) override;
 
-	UEvalGraphEdNode* NodeTemplate;
+	FName NodeTypeName;
 };
diff --git a/Engine/Plugins/Experimental/EvalGraph/Source/EvalGraphEditor/Public/EvalGraph/EvalGraphSchema.h b/Engine/Plugins/Experimental/EvalGraph/Source/EvalGraphEditor/Public/EvalGraph/EvalGraphSchema.h
--- a/Engine/Plugins/Experimental/EvalGraph/Source/EvalGraphEditor/Public/EvalGraph/EvalGraphSchema.h
+++ b/Engine/Plugins/Experimental/EvalGraph/Source/EvalGraphEditor/Public/EvalGraph/EvalGraphSchema.h
@@ -14,6 +14,7 @@ public:
 	UEvalGraphSchema();
 
 	//~ Begin EdGraphSchema Interface
+	virtual const FPinConnectionResponse CanCreateConnection(const UEdGraphPin* PinA, const UEdGraphPin* PinB) const override;
 	virtual void GetContextMenuActions(class UToolMenu* Menu, class UGraphNodeContextMenuContext* Context) const override;
 	virtual void GetGraphContextActions(FGraphContextMenuBuilder& ContextMenuBuilder) const override;
 	//~ End EdGraphSchema Interface
diff --git a/Engine/Source/Runtime/Experimental/EvalGraph/Core/Private/EvalGraph/EvalGraph.cpp b/Engine/Source/Runtime/Experimental/EvalGraph/Core/Private/EvalGraph/EvalGraph.cpp
--- a/Engine/Source/Runtime/Experimental/EvalGraph/Core/Private/EvalGraph/EvalGraph.cpp
+++ b/Engine/Source/Runtime/Experimental/EvalGraph/Core/Private/EvalGraph/EvalGraph.cpp
@@ -50,6 +50,29 @@ namespace Eg
 		Nodes.Remove(Node);
 	}
 
+	void FGraph::ClearConnections(FConnectionBase* Connection)
+	{
+		// Todo(eg) : do this without triggering a invalidation. 
+		//            or implement a better sync for the EdGraph and EgGraph
+		if (Connection->GetDirection() == FPin::EDirection::INPUT)
+		{
+			TArray<FConnectionBase*> BaseOutputs = Connection->GetBaseOutputs();
+			for (FConnectionBase* Output : BaseOutputs)
+			{
+				Disconnect(Connection, Output);
+			}
+		}
+		else if (Connection->GetDirection() == FPin::EDirection::OUTPUT)
+		{
+			TArray<FConnectionBase*> BaseInputs = Connection->GetBaseInputs();
+			for (FConnectionBase* Input : BaseInputs)
+			{
+				Disconnect(Input, Connection);
+			}
+		}
+	}
+
+
 	void FGraph::Connect(FConnectionBase* Input, FConnectionBase* Output)
 	{
 		if (ensure(Input && Output))
diff --git a/Engine/Source/Runtime/Experimental/EvalGraph/Core/Private/EvalGraph/EvalGraphConnectionBase.cpp b/Engine/Source/Runtime/Experimental/EvalGraph/Core/Private/EvalGraph/EvalGraphConnectionBase.cpp
--- a/Engine/Source/Runtime/Experimental/EvalGraph/Core/Private/EvalGraph/EvalGraphConnectionBase.cpp
+++ b/Engine/Source/Runtime/Experimental/EvalGraph/Core/Private/EvalGraph/EvalGraphConnectionBase.cpp
@@ -7,8 +7,9 @@
 
 namespace Eg
 {
-	FConnectionBase::FConnectionBase(FName InType, FName InName, FNode* InOwningNode, FGuid InGuid)
-		: Type(InType)
+	FConnectionBase::FConnectionBase(FPin::EDirection InDirection, FName InType, FName InName, FNode* InOwningNode, FGuid InGuid)
+		: Direction(InDirection)
+		, Type(InType)
 		, Name(InName)
 		, Guid(InGuid)
 		, OwningNode(InOwningNode)
diff --git a/Engine/Source/Runtime/Experimental/EvalGraph/Core/Private/EvalGraph/EvalGraphNode.cpp b/Engine/Source/Runtime/Experimental/EvalGraph/Core/Private/EvalGraph/EvalGraphNode.cpp
--- a/Engine/Source/Runtime/Experimental/EvalGraph/Core/Private/EvalGraph/EvalGraphNode.cpp
+++ b/Engine/Source/Runtime/Experimental/EvalGraph/Core/Private/EvalGraph/EvalGraphNode.cpp
@@ -6,6 +6,26 @@
 
 namespace Eg
 {
+	void FNode::AddBaseInput(FConnectionBase* InPtr)
+	{ 
+		for (FConnectionBase* In : Inputs)
+		{
+			ensureMsgf(!In->GetName().IsEqual(InPtr->GetName()), TEXT("Add Input Failed: Existing Node input already defined with name (%s)"), *InPtr->GetName().ToString());
+		}
+		Inputs.Add(InPtr); 
+	}
+
+	void FNode::AddBaseOutput(FConnectionBase* InPtr)
+	{ 
+		for (FConnectionBase* Out : Outputs)
+		{
+			ensureMsgf(!Out->GetName().IsEqual(InPtr->GetName()), TEXT("Add Output Failed: Existing Node output already defined with name (%s)"), *InPtr->GetName().ToString());
+		}
+		Outputs.Add(InPtr);
+	}
+
+
+
 	TArray<FPin> FNode::GetPins() const
 	{
 		TArray<FPin> RetVal;
@@ -24,5 +44,29 @@ namespace Eg
 		}
 	}
 
+	FConnectionBase* FNode::FindInput(FName InName) const
+	{
+		for (FConnectionBase* Input : Inputs)
+		{
+			if (Input->GetName().IsEqual(InName))
+			{
+				return Input;
+			}
+		}
+		return nullptr;
+	}
+
+	FConnectionBase* FNode::FindOutput(FName InName) const
+	{
+		for (FConnectionBase* Output : Outputs)
+		{
+			if (Output->GetName().IsEqual(InName))
+			{
+				return Output;
+			}
+		}
+		return nullptr;
+	}
+
 }
 
diff --git a/Engine/Source/Runtime/Experimental/EvalGraph/Core/Public/EvalGraph/EvalGraph.h b/Engine/Source/Runtime/Experimental/EvalGraph/Core/Public/EvalGraph/EvalGraph.h
--- a/Engine/Source/Runtime/Experimental/EvalGraph/Core/Public/EvalGraph/EvalGraph.h
+++ b/Engine/Source/Runtime/Experimental/EvalGraph/Core/Public/EvalGraph/EvalGraph.h
@@ -68,6 +68,7 @@ namespace Eg
 
 		void RemoveNode(TSharedPtr<FNode> Node);
 
+		void ClearConnections(FConnectionBase*);
 		void Connect(FConnectionBase* Input, FConnectionBase* Output);
 		void Disconnect(FConnectionBase* Input, FConnectionBase* Output);
 
diff --git a/Engine/Source/Runtime/Experimental/EvalGraph/Core/Public/EvalGraph/EvalGraphConnectionBase.h b/Engine/Source/Runtime/Experimental/EvalGraph/Core/Public/EvalGraph/EvalGraphConnectionBase.h
--- a/Engine/Source/Runtime/Experimental/EvalGraph/Core/Public/EvalGraph/EvalGraphConnectionBase.h
+++ b/Engine/Source/Runtime/Experimental/EvalGraph/Core/Public/EvalGraph/EvalGraphConnectionBase.h
@@ -11,6 +11,17 @@ namespace Eg
 
 	template<class T> inline EVALGRAPHCORE_API FName GraphConnectionTypeName();
 
+	struct FPin
+	{
+		enum class EDirection : uint8 {
+			NONE = 0,
+			INPUT,
+			OUTPUT
+		};
+		EDirection Direction;
+		FName Type;
+		FName Name;
+	};
 
 	//
 	// Input Output Base
@@ -20,6 +31,7 @@ namespace Eg
 	{
 
 	protected:
+		FPin::EDirection Direction;
 		FName Type;
 		FName Name;
 		FGuid  Guid;
@@ -31,9 +43,14 @@ namespace Eg
 
 
 	public:
-		FConnectionBase(FName InType, FName InName, FNode* OwningNode = nullptr, FGuid InGuid = FGuid::NewGuid());
+		FConnectionBase(FPin::EDirection Direction, FName InType, FName InName, FNode* OwningNode = nullptr, FGuid InGuid = FGuid::NewGuid());
 		virtual ~FConnectionBase() {};
 
+		FNode* GetOwningNode() { return OwningNode; }
+		const FNode* GetOwningNode() const { return OwningNode; }
+
+		FPin::EDirection GetDirection() const { return Direction; }
+
 		FName GetType() const { return Type; }
 
 		FGuid GetGuid() const { return Guid; }
diff --git a/Engine/Source/Runtime/Experimental/EvalGraph/Core/Public/EvalGraph/EvalGraphInputOutput.h b/Engine/Source/Runtime/Experimental/EvalGraph/Core/Public/EvalGraph/EvalGraphInputOutput.h
--- a/Engine/Source/Runtime/Experimental/EvalGraph/Core/Public/EvalGraph/EvalGraphInputOutput.h
+++ b/Engine/Source/Runtime/Experimental/EvalGraph/Core/Public/EvalGraph/EvalGraphInputOutput.h
@@ -19,6 +19,7 @@ namespace Eg
 	struct EVALGRAPHCORE_API FInputParameters {
 		FInputParameters(FName InName, FNode* InOwner, T InDefault = T())
 			: Type(GraphConnectionTypeName<T>())
+			, Name(InName)
 			, Owner(InOwner)
 			, Default(InDefault) {}
 		FName Type;
@@ -37,7 +38,7 @@ namespace Eg
 		FOutput<T>* Connection;
 	public:
 		FInput(const FInputParameters<T>& Param, FGuid InGuid = FGuid::NewGuid())
-			: FConnectionBase(Param.Type, Param.Name, Param.Owner, InGuid)
+			: FConnectionBase(FPin::EDirection::INPUT, Param.Type, Param.Name, Param.Owner, InGuid)
 			, Default(Param.Default)
 			, Connection(nullptr)
 		{
@@ -133,7 +134,7 @@ namespace Eg
 	
 	public:
 		FOutput(const FOutputParameters<T>& Param, FGuid InGuid = FGuid::NewGuid())
-			: FConnectionBase(Param.Type, Param.Name, Param.Owner, InGuid)
+			: FConnectionBase(FPin::EDirection::OUTPUT, Param.Type, Param.Name, Param.Owner, InGuid)
 		{
 			Super::AddBaseOutput(Param.Owner, this);
 		}
diff --git a/Engine/Source/Runtime/Experimental/EvalGraph/Core/Public/EvalGraph/EvalGraphNode.h b/Engine/Source/Runtime/Experimental/EvalGraph/Core/Public/EvalGraph/EvalGraphNode.h
--- a/Engine/Source/Runtime/Experimental/EvalGraph/Core/Public/EvalGraph/EvalGraphNode.h
+++ b/Engine/Source/Runtime/Experimental/EvalGraph/Core/Public/EvalGraph/EvalGraphNode.h
@@ -13,17 +13,6 @@ namespace Eg
 		FName Name;
 	};
 
-	struct FPin
-	{
-		enum class EDirection : uint8 {
-			NONE = 0,
-			INPUT,
-			OUTPUT
-		};
-		EDirection Direction;
-		FName Type;
-		FName Name;
-	};
 
 	//
 	// FNode
@@ -54,18 +43,19 @@ namespace Eg
 		FGuid GetGuid() const { return Guid; }
 
 		TArray<FPin> GetPins() const;
+		FConnectionBase* FindInput(FName Name) const;
+		FConnectionBase* FindOutput(FName Name) const;
 
 		virtual void Evaluate(const FContext& Context, FConnectionBase*) { ensure(false); }
 		void InvalidateOutputs();
 
 		virtual void SerializeInternal(FArchive& Ar) {};
 
-	protected:
-		void AddBaseInput(FConnectionBase* InPtr) { Inputs.Add(InPtr); }
+		void AddBaseInput(FConnectionBase* InPtr);
 		const TArray< FConnectionBase* >& GetInputs() const { return Inputs; }
 		TArray< FConnectionBase* >& GetInputs() { return Inputs; }
 
-		void AddBaseOutput(FConnectionBase* InPtr) { Outputs.Add(InPtr); }
+		void AddBaseOutput(FConnectionBase* InPtr);
 		const TArray< FConnectionBase* >& GetOutputs() const { return Outputs; }
 		TArray< FConnectionBase* >& GetOutputs() { return Outputs; }
 
diff --git a/Engine/Source/Runtime/Experimental/EvalGraph/Engine/Private/EvalGraph/EvalGraphEdNode.cpp b/Engine/Source/Runtime/Experimental/EvalGraph/Engine/Private/EvalGraph/EvalGraphEdNode.cpp
--- a/Engine/Source/Runtime/Experimental/EvalGraph/Engine/Private/EvalGraph/EvalGraphEdNode.cpp
+++ b/Engine/Source/Runtime/Experimental/EvalGraph/Engine/Private/EvalGraph/EvalGraphEdNode.cpp
@@ -39,6 +39,69 @@ void UEvalGraphEdNode::AllocateDefaultPins()
 #endif // WITH_EDITOR && !UE_BUILD_SHIPPING
 }
 
+
+FText UEvalGraphEdNode::GetNodeTitle(ENodeTitleType::Type TitleType) const
+{
+	return FText::FromString(GetName());
+}
+
+#if WITH_EDITOR && !(UE_BUILD_SHIPPING || UE_BUILD_TEST)
+void UEvalGraphEdNode::PinConnectionListChanged(UEdGraphPin* Pin)
+{
+	if (ensure(IsBound()))
+	{
+		if (TSharedPtr<Eg::FNode> EgNode = EgGraph->FindBaseNode(EgNodeGuid))
+		{
+			if (Eg::FConnectionBase* ConnectionBaseInput = EgNode->FindInput(FName(Pin->GetName())) )
+			{
+				EgGraph->ClearConnections(ConnectionBaseInput);
+				for (UEdGraphPin* LinkedCon : Pin->LinkedTo)
+				{
+					if (UEvalGraphEdNode* LinkedNode = Cast<UEvalGraphEdNode>(LinkedCon->GetOwningNode()))
+					{
+						if (ensure(LinkedNode->IsBound()))
+						{
+							if (TSharedPtr<Eg::FNode> LinkedEgNode = EgGraph->FindBaseNode(LinkedNode->GetEgNodeGuid()))
+							{
+								if (Eg::FConnectionBase* LinkedConBase = LinkedEgNode->FindOutput(FName(LinkedCon->GetName())))
+								{
+									EgGraph->Connect(ConnectionBaseInput, LinkedConBase);
+								}
+							}
+						}
+					}
+				}
+			}
+			else if (Eg::FConnectionBase* Con = EgNode->FindOutput(FName(Pin->GetName())))
+			{
+				if (Eg::FConnectionBase* ConnectionBaseOutput = EgNode->FindOutput(FName(Pin->GetName())))
+				{
+					EgGraph->ClearConnections(ConnectionBaseOutput);
+					for (UEdGraphPin* LinkedCon : Pin->LinkedTo)
+					{
+						if (UEvalGraphEdNode* LinkedNode = Cast<UEvalGraphEdNode>(LinkedCon->GetOwningNode()))
+						{
+							if (ensure(LinkedNode->IsBound()))
+							{
+								if (TSharedPtr<Eg::FNode> LinkedEgNode = EgGraph->FindBaseNode(LinkedNode->GetEgNodeGuid()))
+								{
+									if (Eg::FConnectionBase* LinkedConBase = LinkedEgNode->FindInput(FName(LinkedCon->GetName())))
+									{
+										EgGraph->Connect(LinkedConBase, ConnectionBaseOutput);
+									}
+								}
+							}
+						}
+					}
+				}
+			}
+		}
+	}
+
+	Super::PinConnectionListChanged(Pin);
+}
+#endif // WITH_EDITOR && !UE_BUILD_SHIPPING
+
 void UEvalGraphEdNode::Serialize(FArchive& Ar)
 {
 	Super::Serialize(Ar);
diff --git a/Engine/Source/Runtime/Experimental/EvalGraph/Engine/Public/EvalGraph/EvalGraphEdNode.h b/Engine/Source/Runtime/Experimental/EvalGraph/Engine/Public/EvalGraph/EvalGraphEdNode.h
--- a/Engine/Source/Runtime/Experimental/EvalGraph/Engine/Public/EvalGraph/EvalGraphEdNode.h
+++ b/Engine/Source/Runtime/Experimental/EvalGraph/Engine/Public/EvalGraph/EvalGraphEdNode.h
@@ -17,17 +17,27 @@ class EVALGRAPHENGINE_API UEvalGraphEdNode : public UEdGraphNode
 
 	FGuid EgNodeGuid;
 	TSharedPtr<Eg::FGraph> EgGraph;
-
 public:
 
 	// UEdGraphNode interface
 	virtual void AllocateDefaultPins();
+	virtual FText GetNodeTitle(ENodeTitleType::Type TitleType) const;
+#if WITH_EDITOR && !(UE_BUILD_SHIPPING || UE_BUILD_TEST)
+	virtual void PinConnectionListChanged(UEdGraphPin* Pin) override;
+#endif // WITH_EDITOR && !UE_BUILD_SHIPPING
 	// End of UEdGraphNode interface
 
+	bool IsBound() { return EgGraph && EgNodeGuid.IsValid(); }
+
+	TSharedPtr<Eg::FGraph> GetEgGraph() { return EgGraph;}
 	void SetEgGraph(TSharedPtr<Eg::FGraph> InEgGraph) { EgGraph = InEgGraph; }
-	void SetEgNode(FGuid InGuid) {EgNodeGuid = InGuid;}
 
+	FGuid GetEgNodeGuid() const { return EgNodeGuid; }
+	void SetEgNodeGuid(FGuid InGuid) { EgNodeGuid = InGuid; }
+
+	// UObject interface
 	void Serialize(FArchive& Ar);
+	// End UObject interface
 
 };
 
